Sequencer: initial commit for the Sequence Validator tool.

diff --git a/Engine/Plugins/MovieScene/SequenceValidator/SequenceValidator.uplugin b/Engine/Plugins/MovieScene/SequenceValidator/SequenceValidator.uplugin
new file mode 100644
--- /dev/null
+++ b/Engine/Plugins/MovieScene/SequenceValidator/SequenceValidator.uplugin
@@ -0,0 +1,32 @@
+{
+	"FileVersion" : 1,
+	"Version" : 1,
+	"VersionName" : "0.1",
+	"FriendlyName" : "Sequence Validator",
+	"Description" : "A tool that provides validation rules for detecting common errors in sequences",
+	"Category" : "Editor",
+	"CreatedBy" : "Epic Games, Inc.",
+	"CreatedByURL" : "https://epicgames.com/",
+	"DocsURL" : "",
+	"MarketplaceURL" : "",
+	"SupportURL" : "",
+	"EnabledByDefault" : true,
+	"CanContainContent" : false,
+	"IsExperimentalVersion" : true,
+	"Installed" : false,
+	"Modules" :
+	[
+		{
+			"Name" : "SequenceValidator",
+			"Type" : "Editor",
+			"LoadingPhase" : "PostEngineInit"
+		}
+	],
+	"Plugins":
+	[
+		{
+			"Name": "LevelSequenceEditor",
+			"Enabled": true
+		}
+	]
+}
diff --git a/Engine/Plugins/MovieScene/SequenceValidator/Source/SequenceValidator/Private/SequenceValidatorCommands.cpp b/Engine/Plugins/MovieScene/SequenceValidator/Source/SequenceValidator/Private/SequenceValidatorCommands.cpp
new file mode 100644
--- /dev/null
+++ b/Engine/Plugins/MovieScene/SequenceValidator/Source/SequenceValidator/Private/SequenceValidatorCommands.cpp
@@ -0,0 +1,37 @@
+// Copyright Epic Games, Inc. All Rights Reserved.
+
+#include "SequenceValidatorCommands.h"
+
+#include "Framework/Commands/Commands.h"
+#include "Framework/Commands/InputChord.h"
+#include "Framework/Commands/UICommandInfo.h"
+#include "SequenceValidatorStyle.h"
+
+#define LOCTEXT_NAMESPACE "SequenceValidatorCommands"
+
+namespace UE::Sequencer
+{
+
+FSequenceValidatorCommands::FSequenceValidatorCommands()
+	: TCommands<FSequenceValidatorCommands>(
+			"SequenceValidator",
+			LOCTEXT("SequenceValidator", "Sequence Validator"),
+			NAME_None,
+			FSequenceValidatorStyle::Get()->GetStyleSetName()
+		)
+{
+}
+
+void FSequenceValidatorCommands::RegisterCommands()
+{
+	UI_COMMAND(
+			StartValidation, "Start Validation", 
+			"Start validating the sequences in the queue",
+			EUserInterfaceActionType::Button,
+			FInputChord());
+}
+
+}  // namespace UE::Sequencer
+
+#undef LOCTEXT_NAMESPACE
+
diff --git a/Engine/Plugins/MovieScene/SequenceValidator/Source/SequenceValidator/Private/SequenceValidatorCommands.h b/Engine/Plugins/MovieScene/SequenceValidator/Source/SequenceValidator/Private/SequenceValidatorCommands.h
new file mode 100644
--- /dev/null
+++ b/Engine/Plugins/MovieScene/SequenceValidator/Source/SequenceValidator/Private/SequenceValidatorCommands.h
@@ -0,0 +1,29 @@
+// Copyright Epic Games, Inc. All Rights Reserved.
+
+#pragma once
+
+#include "CoreTypes.h"
+#include "Framework/Commands/Commands.h"
+
+class FUICommandInfo;
+class UClass;
+struct FInputChord;
+
+namespace UE::Sequencer
+{
+
+class FSequenceValidatorCommands : public TCommands<FSequenceValidatorCommands>
+{
+public:
+
+	FSequenceValidatorCommands();
+
+	virtual void RegisterCommands() override;
+	
+public:
+
+	TSharedPtr<FUICommandInfo> StartValidation;
+};
+
+}  // namespace UE::Sequencer
+
diff --git a/Engine/Plugins/MovieScene/SequenceValidator/Source/SequenceValidator/Private/SequenceValidatorModule.cpp b/Engine/Plugins/MovieScene/SequenceValidator/Source/SequenceValidator/Private/SequenceValidatorModule.cpp
new file mode 100644
--- /dev/null
+++ b/Engine/Plugins/MovieScene/SequenceValidator/Source/SequenceValidator/Private/SequenceValidatorModule.cpp
@@ -0,0 +1,116 @@
+// Copyright Epic Games, Inc. All Rights Reserved.
+
+#include "ISequenceValidatorModule.h"
+
+#include "Editor.h"
+#include "SequenceValidatorCommands.h"
+#include "Validation/Rules/SequenceValidationRule_SectionAlignments.h"
+#include "Validation/Rules/SequenceValidationRule_WholeSectionRanges.h"
+#include "Widgets/SSequenceValidator.h"
+
+DEFINE_LOG_CATEGORY(LogSequenceValidator);
+
+namespace UE::Sequencer
+{
+
+class FSequenceValidatorModule
+	: public ISequenceValidatorModule
+{
+
+protected:
+
+	// IModuleInterface interface.
+	virtual void StartupModule() override;
+	virtual void ShutdownModule() override;
+
+	// ISequenceValidatorModule interface.
+	virtual FSequenceValidationRuleID RegisterValidationRule(FSequenceValidationRuleInfo&& InRuleInfo) override;
+	virtual TConstArrayView<FSequenceValidationRuleInfo> GetValidationRules() const override;
+	virtual void UnregisterValidationRule(FSequenceValidationRuleID InHandle) override;
+
+private:
+
+	void OnPostEngineInit();
+
+	void RegisterBuiltInRules();
+	void UnregisterBuiltInRules();
+
+private:
+
+	TArray<FSequenceValidationRuleInfo> ValidationRules;
+	TArray<FSequenceValidationRuleID> BuiltInValidationRules;
+};
+
+void FSequenceValidatorModule::StartupModule()
+{
+	if (GEditor)
+	{
+		OnPostEngineInit();
+	}
+	else
+	{
+		FCoreDelegates::OnPostEngineInit.AddRaw(this, &FSequenceValidatorModule::OnPostEngineInit);
+	}
+
+	FSequenceValidatorCommands::Register();
+	
+	RegisterBuiltInRules();
+}
+
+void FSequenceValidatorModule::ShutdownModule()
+{
+	SSequenceValidator::UnregisterTabSpawners();
+
+	FSequenceValidatorCommands::Unregister();
+
+	UnregisterBuiltInRules();
+
+	FCoreDelegates::OnPostEngineInit.RemoveAll(this);
+}
+
+void FSequenceValidatorModule::OnPostEngineInit()
+{
+	SSequenceValidator::RegisterTabSpawners();
+}
+
+void FSequenceValidatorModule::RegisterBuiltInRules()
+{
+	BuiltInValidationRules.Add(
+			RegisterValidationRule(FSequenceValidationRule_SectionAlignments::MakeRuleInfo()));
+	BuiltInValidationRules.Add(
+			RegisterValidationRule(FSequenceValidationRule_WholeSectionRanges::MakeRuleInfo()));
+}
+
+void FSequenceValidatorModule::UnregisterBuiltInRules()
+{
+	for (FDelegateHandle Handle : BuiltInValidationRules)
+	{
+		UnregisterValidationRule(Handle);
+	}
+	BuiltInValidationRules.Reset();
+}
+
+FSequenceValidationRuleID FSequenceValidatorModule::RegisterValidationRule(FSequenceValidationRuleInfo&& InRuleInfo)
+{
+	ValidationRules.Add(MoveTemp(InRuleInfo));
+	return ValidationRules.Last().RuleFactory.GetHandle();
+}
+
+TConstArrayView<FSequenceValidationRuleInfo> FSequenceValidatorModule::GetValidationRules() const
+{
+	return ValidationRules;
+}
+
+void FSequenceValidatorModule::UnregisterValidationRule(FDelegateHandle InHandle)
+{
+	ValidationRules.RemoveAll(
+			[=](const FSequenceValidationRuleInfo& RuleInfo) 
+			{
+				return RuleInfo.RuleFactory.GetHandle() == InHandle; 
+			});
+}
+
+}  // namespace UE::Sequencer
+
+IMPLEMENT_MODULE(UE::Sequencer::FSequenceValidatorModule, SequenceValidator);
+
diff --git a/Engine/Plugins/MovieScene/SequenceValidator/Source/SequenceValidator/Private/SequenceValidatorStyle.cpp b/Engine/Plugins/MovieScene/SequenceValidator/Source/SequenceValidator/Private/SequenceValidatorStyle.cpp
new file mode 100644
--- /dev/null
+++ b/Engine/Plugins/MovieScene/SequenceValidator/Source/SequenceValidator/Private/SequenceValidatorStyle.cpp
@@ -0,0 +1,84 @@
+// Copyright Epic Games, Inc. All Rights Reserved.
+
+#include "SequenceValidatorStyle.h"
+
+#include "Interfaces/IPluginManager.h"
+#include "Misc/Paths.h"
+#include "Styling/AppStyle.h"
+#include "Styling/CoreStyle.h"
+#include "Styling/SlateStyleRegistry.h"
+#include "Styling/SlateStyleMacros.h"
+#include "Styling/SlateTypes.h"
+
+namespace UE::Sequencer
+{
+
+TSharedPtr<FSequenceValidatorStyle> FSequenceValidatorStyle::Singleton;
+
+FSequenceValidatorStyle::FSequenceValidatorStyle()
+	: FSlateStyleSet("SequenceValidatorStyle")
+{
+	const FVector2D Icon16x16(16.0f, 16.0f);
+
+	const FString ContentDir = IPluginManager::Get().FindPlugin(TEXT("SequenceValidator"))->GetContentDir();
+	SetContentRoot(ContentDir);
+	SetCoreContentRoot(FPaths::EngineContentDir() / TEXT("Slate"));
+
+	const ISlateStyle& CoreStyle = FAppStyle::Get();
+	const FTextBlockStyle NormalText = CoreStyle.GetWidgetStyle<FTextBlockStyle>("NormalText");
+	const FTableRowStyle AlternatingTableRowStyle = CoreStyle.GetWidgetStyle<FTableRowStyle>("TableView.AlternatingRow");
+
+	// Validation rules icons.
+	Set("ValidationRule", new IMAGE_BRUSH_SVG("Icons/ValidationRule", Icon16x16));
+
+	Set("ValidationRule.RowStyle", AlternatingTableRowStyle);
+
+	// Validation results icons.
+	Set("ValidationResult.Key", new IMAGE_BRUSH("Icons/ValidationResult-Key", Icon16x16));
+	Set("ValidationResult.Section", new IMAGE_BRUSH("Icons/ValidationResult-Section", Icon16x16));
+	Set("ValidationResult.Track", new IMAGE_BRUSH("Icons/ValidationResult-Track", Icon16x16));
+	Set("ValidationResult.Sequence", new IMAGE_BRUSH("Icons/ValidationResult-Sequence", Icon16x16));
+
+	Set("ValidationResult.RowStyle", AlternatingTableRowStyle);
+
+	// Command icons.
+	Set("SequenceValidator.QueueSequence", new CORE_IMAGE_BRUSH_SVG("Starship/Common/plus", Icon16x16));
+	Set("SequenceValidator.StartValidation", new IMAGE_BRUSH_SVG("Icons/SequenceValidation", Icon16x16));
+
+	// Tab icons.
+	Set("SequenceValidator.TabIcon", new IMAGE_BRUSH_SVG("Icons/SequenceValidation", Icon16x16));
+
+	// Text styles.
+
+	Set("SequenceValidator.PanelTitleText", FTextBlockStyle(NormalText)
+		.SetFont(DEFAULT_FONT("Bold", 18))
+		.SetColorAndOpacity(FLinearColor(1.0f, 1.0f, 1.0f))
+		.SetHighlightColor(FLinearColor(1.0f, 1.0f, 1.0f))
+		.SetShadowOffset(FVector2f::UnitVector)
+		.SetShadowColorAndOpacity(FLinearColor(0, 0, 0, 0.9f)));
+
+	Set("ValidatorRule.RuleNameText", FTextBlockStyle(NormalText)
+		.SetFont(DEFAULT_FONT("Bold", 12)));
+
+	Set("ValidatorRule.RuleDescriptionText", FTextBlockStyle(NormalText)
+		.SetFont(DEFAULT_FONT("Regular", 10)));
+
+	FSlateStyleRegistry::RegisterSlateStyle(*this);
+}
+
+FSequenceValidatorStyle::~FSequenceValidatorStyle()
+{
+	FSlateStyleRegistry::UnRegisterSlateStyle(*this);
+}
+
+TSharedRef<FSequenceValidatorStyle> FSequenceValidatorStyle::Get()
+{
+	if (!Singleton.IsValid())
+	{
+		Singleton = MakeShareable(new FSequenceValidatorStyle);
+	}
+	return Singleton.ToSharedRef();
+}
+
+}  // namespace UE::Sequencer
+
diff --git a/Engine/Plugins/MovieScene/SequenceValidator/Source/SequenceValidator/Private/SequenceValidatorStyle.h b/Engine/Plugins/MovieScene/SequenceValidator/Source/SequenceValidator/Private/SequenceValidatorStyle.h
new file mode 100644
--- /dev/null
+++ b/Engine/Plugins/MovieScene/SequenceValidator/Source/SequenceValidator/Private/SequenceValidatorStyle.h
@@ -0,0 +1,28 @@
+// Copyright Epic Games, Inc. All Rights Reserved.
+
+#pragma once
+
+#include "Styling/SlateStyle.h"
+
+namespace UE::Sequencer
+{
+
+/**
+ * Implements the visual style of the sequence validator.
+ */
+class FSequenceValidatorStyle final : public FSlateStyleSet
+{
+public:
+	
+	FSequenceValidatorStyle();
+	virtual ~FSequenceValidatorStyle();
+
+	static TSharedRef<FSequenceValidatorStyle> Get();
+
+private:
+
+	static TSharedPtr<FSequenceValidatorStyle> Singleton;
+};
+
+}  // namespace UE::Sequencer
+
diff --git a/Engine/Plugins/MovieScene/SequenceValidator/Source/SequenceValidator/Private/Validation/Rules/SequenceValidationRule_SectionAlignments.cpp b/Engine/Plugins/MovieScene/SequenceValidator/Source/SequenceValidator/Private/Validation/Rules/SequenceValidationRule_SectionAlignments.cpp
new file mode 100644
--- /dev/null
+++ b/Engine/Plugins/MovieScene/SequenceValidator/Source/SequenceValidator/Private/Validation/Rules/SequenceValidationRule_SectionAlignments.cpp
@@ -0,0 +1,252 @@
+// Copyright Epic Games, Inc. All Rights Reserved.
+
+#include "Validation/Rules/SequenceValidationRule_SectionAlignments.h"
+
+#include "Algo/BinarySearch.h"
+#include "Internationalization/Internationalization.h"
+#include "MovieScene.h"
+#include "MovieSceneTimeHelpers.h"
+#include "MovieSceneSection.h"
+#include "MovieSceneSequence.h"
+#include "MovieSceneSequenceVisitor.h"
+#include "MovieSceneTrack.h"
+#include "Sections/MovieSceneSubSection.h"
+#include "Tracks/MovieSceneCameraCutTrack.h"
+#include "Tracks/MovieSceneCinematicShotTrack.h"
+#include "Tracks/MovieSceneSubTrack.h"
+#include "Validation/SequenceValidationResult.h"
+
+#define LOCTEXT_NAMESPACE "SequenceValidationRule_SectionAlignments"
+
+namespace UE::Sequencer
+{
+
+struct FImportantFrameTime
+{
+	FFrameTime FrameTime;
+	FText Label;
+};
+
+bool operator< (const FImportantFrameTime& A, const FImportantFrameTime& B)
+{
+	return A.FrameTime < B.FrameTime;
+}
+
+bool operator== (const FImportantFrameTime& A, const FImportantFrameTime& B)
+{
+	return A.FrameTime == B.FrameTime;
+}
+
+uint32 GetTypeHash(const FImportantFrameTime& A)
+{
+	return HashCombine(::GetTypeHash(A.FrameTime.FrameNumber.Value), ::GetTypeHash(A.FrameTime.GetSubFrame()));
+}
+
+class FSectionAlignmentsImportantTimeVisitor : public UE::MovieScene::ISequenceVisitor
+{
+	using FSubSequenceSpace = UE::MovieScene::FSubSequenceSpace;
+
+public:
+
+	virtual void VisitTrack(UMovieSceneTrack* Track, const FGuid& ObjectBindingID, const FSubSequenceSpace& LocalSpace) override
+	{
+		if (UMovieSceneCameraCutTrack* CameraCutTrack = Cast<UMovieSceneCameraCutTrack>(Track))
+		{
+			for (UMovieSceneSection* CameraCutSection : CameraCutTrack->GetAllSections())
+			{
+				TRange<FFrameNumber> TimeRange = CameraCutSection->GetTrueRange();
+				AddImportantTimes(
+						TimeRange, LocalSpace,
+						LOCTEXT("CameraCutBound", "Camera Cut Bound"),
+						LOCTEXT("CameraCutBound", "Camera Cut Bound"));
+			}
+		}
+	}
+
+public:
+
+	void AddImportantTimes(
+			const TRange<FFrameNumber>& InTimeRange, const FSubSequenceSpace& LocalSpace,
+			const FText& LowerTimeLabel, const FText& UpperTimeLabel)
+	{
+		FMovieSceneInverseSequenceTransform SequenceToRootTransform = LocalSpace.RootToSequenceTransform.Inverse();
+		if (InTimeRange.HasLowerBound())
+		{
+			TOptional<FFrameTime> RootTime = SequenceToRootTransform.TryTransformTime(InTimeRange.GetLowerBoundValue());
+			if (RootTime.IsSet())
+			{
+				ImportantTimes.FindOrAdd({ RootTime.GetValue(), LowerTimeLabel });
+			}
+		}
+		if (InTimeRange.HasUpperBound())
+		{
+			TOptional<FFrameTime> RootTime = SequenceToRootTransform.TryTransformTime(InTimeRange.GetUpperBoundValue());
+			if (RootTime.IsSet())
+			{
+				ImportantTimes.FindOrAdd({ RootTime.GetValue(), UpperTimeLabel });
+			}
+		}
+	}
+
+	void GetImportantTimes(TArray<FImportantFrameTime>& OutImportantTimes)
+	{
+		OutImportantTimes = ImportantTimes.Array();
+	}
+
+private:
+
+	TSet<FImportantFrameTime> ImportantTimes;
+};
+
+struct FSectionAlignmentsValidationVisitor : public UE::MovieScene::ISequenceVisitor
+{
+	using FSubSequenceSpace = UE::MovieScene::FSubSequenceSpace;
+
+public:
+
+	FSectionAlignmentsValidationVisitor(TArrayView<FImportantFrameTime> InImportantTimes, FSequenceValidationResults& OutResults)
+		: ImportantTimes(InImportantTimes)
+		, Results(OutResults)
+	{
+	}
+
+	void Initialize(const UMovieSceneSequence* InSequence)
+	{
+		UMovieScene* MovieScene = InSequence->GetMovieScene();
+
+		ImportantTimes.Sort();
+		CurrentSubSectionTrail.Reserve(4);
+
+		// We will emit warnings for sections that start/end within 2 frames of an "important time".
+		const FFrameRate& DisplayRate = MovieScene->GetDisplayRate();
+		const FFrameRate& TickResolution = MovieScene->GetTickResolution();
+		Threshold = FFrameRate::TransformTime(FFrameTime(2), DisplayRate, TickResolution);
+	}
+
+	virtual void VisitStartSubSequence(UMovieSceneSubSection* SubSection, const FGuid& ObjectBindingID, const FSubSequenceSpace& LocalSpace) override
+	{
+		CurrentSubSectionTrail.Add(SubSection);
+	}
+
+	virtual void VisitEndSubSequence(UMovieSceneSubSection* SubSection, const FGuid& ObjectBindingID, const FSubSequenceSpace& LocalSpace) override
+	{
+		ensure(CurrentSubSectionTrail.Num() > 0 && CurrentSubSectionTrail.Last() == SubSection);
+		CurrentSubSectionTrail.Pop(EAllowShrinking::No);
+	}
+
+	virtual void VisitSection(UMovieSceneTrack* Track, UMovieSceneSection* Section, const FGuid& ObjectBindingID, const FSubSequenceSpace& LocalSpace) override
+	{
+		TRange<FFrameNumber> TimeRange = Section->GetTrueRange();
+		FMovieSceneInverseSequenceTransform SequenceToRootTransform = LocalSpace.RootToSequenceTransform.Inverse();
+		if (TimeRange.HasLowerBound())
+		{
+			MaybeWarnAbout(Results, SequenceToRootTransform, TimeRange.GetLowerBoundValue(), Section);
+		}
+		if (TimeRange.HasUpperBound())
+		{
+			MaybeWarnAbout(Results, SequenceToRootTransform, TimeRange.GetUpperBoundValue(), Section);
+		}
+	}
+
+private:
+
+	void MaybeWarnAbout(FSequenceValidationResults& OutResults, const FMovieSceneInverseSequenceTransform& SequenceToRootTransform, FFrameTime InTime, UObject* InTarget = nullptr)
+	{
+		// Convert the local time to a root sequence time, so that we can compare it to all the important times that
+		// are also in root time space.
+		TOptional<FFrameTime> OptRootTime = SequenceToRootTransform.TryTransformTime(InTime);
+		if (!OptRootTime.IsSet())
+		{
+			return;
+		}
+
+		const FFrameTime RootTime = OptRootTime.GetValue();
+
+		const int32 NextTimeIndex = Algo::LowerBound(ImportantTimes, FImportantFrameTime{ RootTime });
+		const int32 PreviousTimeIndex = NextTimeIndex - 1;
+
+		bool bDoWarn = false;
+		FText TimeLabel;
+		if (ImportantTimes.IsValidIndex(PreviousTimeIndex))
+		{
+			const FImportantFrameTime& PreviousTime = ImportantTimes[PreviousTimeIndex];
+			if (RootTime != PreviousTime.FrameTime && FMath::Abs((RootTime - PreviousTime.FrameTime)) <= Threshold)
+			{
+				bDoWarn = true;
+				TimeLabel = PreviousTime.Label;
+			}
+		}
+		if (ImportantTimes.IsValidIndex(NextTimeIndex))
+		{
+			const FImportantFrameTime& NextTime = ImportantTimes[NextTimeIndex];
+			if (RootTime != NextTime.FrameTime && FMath::Abs((RootTime - NextTime.FrameTime)) <= Threshold)
+			{
+				bDoWarn = true;
+				TimeLabel = NextTime.Label;
+			}
+		}
+
+		if (bDoWarn)
+		{
+			TSharedRef<FSequenceValidationResult> NewResult = MakeShared<FSequenceValidationResult>(EMessageSeverity::Warning, InTarget);
+			NewResult->SetSubSectionTrail(CurrentSubSectionTrail);
+			NewResult->SetLocalTime(InTime);
+			NewResult->SetUserMessage(FText::Format(
+						LOCTEXT("AlignmentWarningFmt", "Section bound is within two frames of {0}."),
+						TimeLabel));
+			OutResults.AddResult(NewResult);
+		}
+	}
+
+private:
+
+	FFrameTime Threshold;
+	TArray<FImportantFrameTime> ImportantTimes;
+	TArray<UMovieSceneSubSection*> CurrentSubSectionTrail;
+	FSequenceValidationResults& Results;
+};
+
+FSequenceValidationRuleInfo FSequenceValidationRule_SectionAlignments::MakeRuleInfo()
+{
+	FSequenceValidationRuleInfo RuleInfo;
+	RuleInfo.RuleName = LOCTEXT("RuleName", "Section Alignments");
+	RuleInfo.RuleDescription = LOCTEXT("RuleDescription", "Check that sections are not off by a couple frames from notable frames.");
+	RuleInfo.RuleFactory = FOnCreateSequenceValidationRule::CreateLambda([]() { return MakeShared<FSequenceValidationRule_SectionAlignments>(); });
+	return RuleInfo;
+}
+
+void FSequenceValidationRule_SectionAlignments::OnRun(const UMovieSceneSequence* InSequence, FSequenceValidationResults& OutResults) const
+{
+	UMovieScene* MovieScene = InSequence->GetMovieScene();
+
+	FSectionAlignmentsImportantTimeVisitor Visitor;
+
+	// Add the root sequence's playback start and end as important times.
+	Visitor.AddImportantTimes(
+			MovieScene->GetPlaybackRange(), UE::MovieScene::FSubSequenceSpace(),
+			LOCTEXT("PlaybackRangeLowerBound", "Playback Range's Lower Bound"),
+			LOCTEXT("PlaybackRangeUpperBound", "Playback Range's Upper Bound"));
+
+	// Find important times in the whole hierarchy.
+	UE::MovieScene::FSequenceVisitParams VisitParams;
+	VisitParams.bVisitRootTracks = true;
+	VisitParams.bVisitSubSequences = true;
+	UE::MovieScene::VisitSequence(const_cast<UMovieSceneSequence*>(InSequence), VisitParams, Visitor);
+
+	TArray<FImportantFrameTime> ImportantTimes;
+	Visitor.GetImportantTimes(ImportantTimes);
+
+	// Re-visit the sequence hierarchy and report on things that are off.
+	VisitParams.bVisitObjectBindings = true;
+	VisitParams.bVisitSections = true;
+	VisitParams.bVisitTracks = true;
+
+	FSectionAlignmentsValidationVisitor Validator(ImportantTimes, OutResults);
+	Validator.Initialize(InSequence);
+	UE::MovieScene::VisitSequence(const_cast<UMovieSceneSequence*>(InSequence), VisitParams, Validator);
+}
+
+}  // namespace UE::Sequencer
+
+#undef LOCTEXT_NAMESPACE
+
diff --git a/Engine/Plugins/MovieScene/SequenceValidator/Source/SequenceValidator/Private/Validation/Rules/SequenceValidationRule_SectionAlignments.h b/Engine/Plugins/MovieScene/SequenceValidator/Source/SequenceValidator/Private/Validation/Rules/SequenceValidationRule_SectionAlignments.h
new file mode 100644
--- /dev/null
+++ b/Engine/Plugins/MovieScene/SequenceValidator/Source/SequenceValidator/Private/Validation/Rules/SequenceValidationRule_SectionAlignments.h
@@ -0,0 +1,26 @@
+// Copyright Epic Games, Inc. All Rights Reserved.
+
+#pragma once
+
+#include "Validation/SequenceValidationRule.h"
+
+namespace UE::Sequencer
+{
+
+/**
+ * A validation rule that checks that sections throughout the sequence hierarchy are correctly aligned
+ * to "important times" such as camera cuts and start/end times.
+ */
+class FSequenceValidationRule_SectionAlignments : public FSequenceValidationRule
+{
+public:
+
+	static FSequenceValidationRuleInfo MakeRuleInfo();
+
+protected:
+
+	virtual void OnRun(const UMovieSceneSequence* InSequence, FSequenceValidationResults& OutResults) const override;
+};
+
+}  // namespace UE::Sequencer
+
diff --git a/Engine/Plugins/MovieScene/SequenceValidator/Source/SequenceValidator/Private/Validation/Rules/SequenceValidationRule_WholeSectionRanges.cpp b/Engine/Plugins/MovieScene/SequenceValidator/Source/SequenceValidator/Private/Validation/Rules/SequenceValidationRule_WholeSectionRanges.cpp
new file mode 100644
--- /dev/null
+++ b/Engine/Plugins/MovieScene/SequenceValidator/Source/SequenceValidator/Private/Validation/Rules/SequenceValidationRule_WholeSectionRanges.cpp
@@ -0,0 +1,106 @@
+// Copyright Epic Games, Inc. All Rights Reserved.
+
+#include "Validation/Rules/SequenceValidationRule_WholeSectionRanges.h"
+
+#include "Internationalization/Internationalization.h"
+#include "MovieScene.h"
+#include "MovieSceneSection.h"
+#include "MovieSceneSequence.h"
+#include "MovieSceneSequenceVisitor.h"
+#include "Validation/SequenceValidationResult.h"
+
+#define LOCTEXT_NAMESPACE "SequenceValidationRule_WholeSectionRanges"
+
+namespace UE::Sequencer
+{
+
+class FWholeSectionRangeVisitor : public UE::MovieScene::ISequenceVisitor
+{
+	using FSubSequenceSpace = UE::MovieScene::FSubSequenceSpace;
+
+public:
+
+	FWholeSectionRangeVisitor(FSequenceValidationResults& OutResults)
+		: Results(OutResults)
+	{
+	}
+
+	virtual void VisitStartSubSequence(UMovieSceneSubSection* SubSection, const FGuid& ObjectBindingID, const FSubSequenceSpace& LocalSpace) override
+	{
+		CurrentSubSectionTrail.Add(SubSection);
+	}
+
+	virtual void VisitEndSubSequence(UMovieSceneSubSection* SubSection, const FGuid& ObjectBindingID, const FSubSequenceSpace& LocalSpace) override
+	{
+		ensure(CurrentSubSectionTrail.Num() > 0 && CurrentSubSectionTrail.Last() == SubSection);
+		CurrentSubSectionTrail.Pop(EAllowShrinking::No);
+	}
+
+	virtual void VisitSection(UMovieSceneTrack* Track, UMovieSceneSection* Section, const FGuid& ObjectBindingID, const FSubSequenceSpace& LocalSpace) override
+	{
+		UMovieScene* MovieScene = Track->GetTypedOuter<UMovieScene>();
+		const FFrameRate& DisplayRate = MovieScene->GetDisplayRate();
+		const FFrameRate& TickResolution = MovieScene->GetTickResolution();
+
+		TRange<FFrameNumber> SectionRange = Section->GetRange();
+		if (SectionRange.HasLowerBound())
+		{
+			const FFrameTime StartTime = ConvertFrameTime(SectionRange.GetLowerBoundValue(), TickResolution, DisplayRate);
+			if (StartTime.GetSubFrame() != 0.f)
+			{
+				AddWarning(Section, StartTime);
+			}
+		}
+		if (SectionRange.HasUpperBound())
+		{
+			const FFrameTime EndTime = ConvertFrameTime(SectionRange.GetUpperBoundValue(), TickResolution, DisplayRate);
+			if (EndTime.GetSubFrame() != 0.f)
+			{
+				AddWarning(Section, EndTime);
+			}
+		}
+	}
+
+private:
+
+	void AddWarning(UMovieSceneSection* Section, FFrameTime LocalTime)
+	{
+		TSharedRef<FSequenceValidationResult> NewResult = MakeShared<FSequenceValidationResult>(EMessageSeverity::Warning, Section);
+		NewResult->SetSubSectionTrail(CurrentSubSectionTrail);
+		NewResult->SetLocalTime(LocalTime);
+		NewResult->SetUserMessage(LOCTEXT("WholeFrameWarningFmt", "Section bound is not on a whole frame."));
+		Results.AddResult(NewResult);
+	}
+
+private:
+
+	TArray<UMovieSceneSubSection*> CurrentSubSectionTrail;
+	FSequenceValidationResults& Results;
+};
+
+FSequenceValidationRuleInfo FSequenceValidationRule_WholeSectionRanges::MakeRuleInfo()
+{
+	FSequenceValidationRuleInfo RuleInfo;
+	RuleInfo.RuleName = LOCTEXT("RuleName", "Whole Section Ranges");
+	RuleInfo.RuleDescription = LOCTEXT("RuleDescription", "Check that sections start and end on whole frames.");
+	RuleInfo.RuleFactory = FOnCreateSequenceValidationRule::CreateLambda([]() { return MakeShared<FSequenceValidationRule_WholeSectionRanges>(); });
+	return RuleInfo;
+}
+
+void FSequenceValidationRule_WholeSectionRanges::OnRun(const UMovieSceneSequence* InSequence, FSequenceValidationResults& OutResults) const
+{
+	UE::MovieScene::FSequenceVisitParams VisitParams;
+	VisitParams.bVisitRootTracks = true;
+	VisitParams.bVisitObjectBindings = true;
+	VisitParams.bVisitSubSequences = true;
+	VisitParams.bVisitTracks = true;
+	VisitParams.bVisitSections = true;
+
+	FWholeSectionRangeVisitor Visitor(OutResults);
+	UE::MovieScene::VisitSequence(const_cast<UMovieSceneSequence*>(InSequence), VisitParams, Visitor);
+}
+
+}  // namespace UE::Sequencer
+
+#undef LOCTEXT_NAMESPACE
+
diff --git a/Engine/Plugins/MovieScene/SequenceValidator/Source/SequenceValidator/Private/Validation/Rules/SequenceValidationRule_WholeSectionRanges.h b/Engine/Plugins/MovieScene/SequenceValidator/Source/SequenceValidator/Private/Validation/Rules/SequenceValidationRule_WholeSectionRanges.h
new file mode 100644
--- /dev/null
+++ b/Engine/Plugins/MovieScene/SequenceValidator/Source/SequenceValidator/Private/Validation/Rules/SequenceValidationRule_WholeSectionRanges.h
@@ -0,0 +1,25 @@
+// Copyright Epic Games, Inc. All Rights Reserved.
+
+#pragma once
+
+#include "Validation/SequenceValidationRule.h"
+
+namespace UE::Sequencer
+{
+
+/**
+ * A validation rule that checks that sections have start/end times that are on whole frames.
+ */
+class FSequenceValidationRule_WholeSectionRanges : public FSequenceValidationRule
+{
+public:
+
+	static FSequenceValidationRuleInfo MakeRuleInfo();
+
+protected:
+
+	virtual void OnRun(const UMovieSceneSequence* InSequence, FSequenceValidationResults& OutResults) const override;
+};
+
+}  // namespace UE::Sequencer
+
diff --git a/Engine/Plugins/MovieScene/SequenceValidator/Source/SequenceValidator/Private/Validation/SequenceValidationResult.cpp b/Engine/Plugins/MovieScene/SequenceValidator/Source/SequenceValidator/Private/Validation/SequenceValidationResult.cpp
new file mode 100644
--- /dev/null
+++ b/Engine/Plugins/MovieScene/SequenceValidator/Source/SequenceValidator/Private/Validation/SequenceValidationResult.cpp
@@ -0,0 +1,100 @@
+// Copyright Epic Games, Inc. All Rights Reserved.
+
+#include "Validation/SequenceValidationResult.h"
+
+#include "Sections/MovieSceneSubSection.h"
+#include "Templates/SharedPointer.h"
+
+namespace UE::Sequencer
+{
+
+bool FSequenceValidationResult::GetSubSectionTrail(TArray<UMovieSceneSubSection*>& OutTrail) const
+{
+	bool bSuccess = true;
+	for (TWeakObjectPtr<UMovieSceneSubSection> WeakSubSection : WeakSubSectionTrail)
+	{
+		if (UMovieSceneSubSection* SubSection = WeakSubSection.Get())
+		{
+			OutTrail.Add(SubSection);
+		}
+		else
+		{
+			bSuccess = false;
+			break;
+		}
+	}
+	if (!bSuccess)
+	{
+		OutTrail.Reset();
+		return false;
+	}
+	return true;
+}
+
+void FSequenceValidationResult::SetSubSectionTrail(TArrayView<UMovieSceneSubSection*> InTrail)
+{
+	WeakSubSectionTrail.Reset();
+	for (UMovieSceneSubSection* SubSection : InTrail)
+	{
+		WeakSubSectionTrail.Add(SubSection);
+	}
+}
+
+TSharedPtr<FSequenceValidationResult> FSequenceValidationResult::GetRoot()
+{
+	TSharedPtr<FSequenceValidationResult> Current = SharedThis(this);
+	while (TSharedPtr<FSequenceValidationResult> Parent = Current->GetParent())
+	{
+		Current = Parent;
+	}
+	return Current;
+}
+
+void FSequenceValidationResult::AddChild(TSharedRef<FSequenceValidationResult> InChild)
+{
+	if (ensure(!InChild->WeakParent.IsValid()))
+	{
+		InChild->WeakParent = SharedThis(this);
+		Children.Add(InChild);
+	}
+}
+
+void FSequenceValidationResult::AppendChildren(TConstArrayView<TSharedPtr<FSequenceValidationResult>> InChildren)
+{
+	for (TSharedPtr<FSequenceValidationResult> Child : InChildren)
+	{
+		if (Child && ensure(!Child->WeakParent.IsValid()))
+		{
+			Child->WeakParent = SharedThis(this);
+			Children.Add(Child);
+		}
+	}
+}
+
+void FSequenceValidationResults::AddResult(TSharedRef<FSequenceValidationResult> InResult)
+{
+	ValidationResults.Add(InResult);
+}
+
+void FSequenceValidationResults::AppendResults(const FSequenceValidationResults& InResults)
+{
+	ValidationResults.Append(InResults.ValidationResults);
+}
+
+void FSequenceValidationResults::AppendResults(TConstArrayView<TSharedPtr<FSequenceValidationResult>> InResults)
+{
+	ValidationResults.Append(InResults);
+}
+
+TArrayView<const TSharedPtr<FSequenceValidationResult>> FSequenceValidationResults::GetResults() const
+{
+	return ValidationResults;
+}
+
+void FSequenceValidationResults::Reset()
+{
+	ValidationResults.Reset();
+}
+
+}  // namespace UE::Sequencer
+
diff --git a/Engine/Plugins/MovieScene/SequenceValidator/Source/SequenceValidator/Private/Validation/SequenceValidationRule.cpp b/Engine/Plugins/MovieScene/SequenceValidator/Source/SequenceValidator/Private/Validation/SequenceValidationRule.cpp
new file mode 100644
--- /dev/null
+++ b/Engine/Plugins/MovieScene/SequenceValidator/Source/SequenceValidator/Private/Validation/SequenceValidationRule.cpp
@@ -0,0 +1,14 @@
+// Copyright Epic Games, Inc. All Rights Reserved.
+
+#include "Validation/SequenceValidationRule.h"
+
+namespace UE::Sequencer
+{
+
+void FSequenceValidationRule::Run(const UMovieSceneSequence* InSequence, FSequenceValidationResults& OutResults) const
+{
+	OnRun(InSequence, OutResults);
+}
+
+}  // namespace UE::Sequencer
+
diff --git a/Engine/Plugins/MovieScene/SequenceValidator/Source/SequenceValidator/Private/Validation/SequenceValidator.cpp b/Engine/Plugins/MovieScene/SequenceValidator/Source/SequenceValidator/Private/Validation/SequenceValidator.cpp
new file mode 100644
--- /dev/null
+++ b/Engine/Plugins/MovieScene/SequenceValidator/Source/SequenceValidator/Private/Validation/SequenceValidator.cpp
@@ -0,0 +1,392 @@
+// Copyright Epic Games, Inc. All Rights Reserved.
+
+#include "Validation/SequenceValidator.h"
+
+#include "Algo/AllOf.h"
+#include "ISequenceValidatorModule.h"
+#include "Modules/ModuleManager.h"
+#include "MovieSceneFwd.h"
+#include "Tasks/Task.h"
+#include "Tasks/TaskPrivate.h"
+#include "Validation/SequenceValidationResult.h"
+#include "Validation/SequenceValidationRule.h"
+
+#define LOCTEXT_NAMESPACE "SequenceValidator"
+
+namespace UE::Sequencer
+{
+
+static TAutoConsoleVariable<int32> CVarSequenceValidationMaxConcurrentTasks(
+	TEXT("SequenceValidation.MaxConcurrentTasks"),
+	-1,
+	TEXT("The maximum number of async tasks running sequence validation rules at the same time.")
+	TEXT("Default: -1 (which means to use the number of cores available."),
+	ECVF_Default);
+
+namespace Internal
+{
+
+class FSequenceValidatorTaskScheduler;
+
+/**
+ * Class responsible for running a single validation rule on a single sequence.
+ */
+struct FSequenceValidatorTask
+{
+	const UMovieSceneSequence* Sequence = nullptr;
+	TSharedPtr<const FSequenceValidationRule> ValidationRule;
+
+	FSequenceValidationResults TaskResults;
+
+	FSequenceValidatorTask(const UMovieSceneSequence* InSequence, TSharedPtr<const FSequenceValidationRule> InRule)
+		: Sequence(InSequence)
+		, ValidationRule(InRule)
+	{
+	}
+
+	void Run();
+};
+
+/**
+ * A class that is able to run validation tasks dispatched by the FSequenceValidatorTaskScheduler.
+ */
+class FSequenceValidatorTaskRunner
+{
+public:
+
+	FSequenceValidatorTaskRunner(int32 InRunnerIndex) : RunnerIndex(InRunnerIndex) {}
+
+	void Start(FSequenceValidatorTaskScheduler& InScheduler);
+
+	bool IsRunning() const { return bIsRunning; }
+
+private:
+
+	void ThreadedRun(FSequenceValidatorTaskScheduler& InScheduler);
+
+private:
+
+	int32 RunnerIndex = INDEX_NONE;
+	bool bIsRunning = false;
+};
+
+/**
+ * Main class for running a set of validation rules on a set of sequences in a parallel fashion.
+ */
+class FSequenceValidatorTaskScheduler
+{
+public:
+
+	FSequenceValidatorTaskScheduler(TArrayView<FSequenceValidationRuleInfo> InRuleInfos)
+		: ValidationRuleInfos(InRuleInfos)
+	{
+	}
+
+	~FSequenceValidatorTaskScheduler()
+	{
+		ensure(!MainSignal);
+	}
+
+	void RunQueue(TArrayView<UMovieSceneSequence*> InQueue, FSimpleDelegate&& InOnFinished);
+
+	const FSequenceValidationResults& GetResults() const { return Results; }
+
+private:
+
+	FSequenceValidatorTask* AcquireNextTask(int32 InRunnerIndex);
+	void OnRunnerFinished(int32 InRunnerIndex);
+	void GatherResults(TArrayView<UMovieSceneSequence*> InQueue);
+
+private:
+
+	TArray<FSequenceValidationRuleInfo> ValidationRuleInfos;
+	FSequenceValidationResults Results;
+
+	// Transient fields for running a queue.
+
+	TArray<FSequenceValidatorTask> Tasks;
+
+	int32 NextTaskIndex = 0;
+	int32 NumRemainingRunners = 0;
+	FEvent* MainSignal = nullptr;
+
+	friend class FSequenceValidatorTaskRunner;
+};
+
+void FSequenceValidatorTask::Run()
+{
+	if (Sequence && ValidationRule)
+	{
+		ValidationRule->Run(Sequence, TaskResults);
+	}
+}
+
+void FSequenceValidatorTaskRunner::Start(FSequenceValidatorTaskScheduler& InScheduler)
+{
+	UE_LOG(LogSequenceValidator, Log, TEXT("Starting sequence validator runner %d"), RunnerIndex);
+	bIsRunning = true;
+
+	UE::Tasks::Launch(
+			TEXT("SequenceValidatorTaskRunner"),
+			[this, &InScheduler]() { this->ThreadedRun(InScheduler); },
+			UE::Tasks::ETaskPriority::Default);
+}
+
+void FSequenceValidatorTaskRunner::ThreadedRun(FSequenceValidatorTaskScheduler& InScheduler)
+{
+	for (;;)
+	{
+		FSequenceValidatorTask* NextTask = InScheduler.AcquireNextTask(RunnerIndex);
+		if (NextTask)
+		{
+			NextTask->Run();
+		}
+		else
+		{
+			break;
+		}
+	}
+
+	UE_LOG(LogSequenceValidator, Log, TEXT("Stopping sequence validator runner %d"), RunnerIndex);
+	bIsRunning = false;
+	InScheduler.OnRunnerFinished(RunnerIndex);
+}
+
+void FSequenceValidatorTaskScheduler::RunQueue(TArrayView<UMovieSceneSequence*> InQueue, FSimpleDelegate&& InOnFinished)
+{
+	ensure(!MainSignal);
+	MainSignal = FPlatformProcess::GetSynchEventFromPool();
+
+	// Create a task for each sequence/rule pair, for rules that are enabled.
+	Tasks.Reset();
+	for (const UMovieSceneSequence* Sequence : InQueue)
+	{
+		for (const FSequenceValidationRuleInfo& RuleInfo : ValidationRuleInfos)
+		{
+			if (RuleInfo.bIsEnabled)
+			{
+				TSharedRef<FSequenceValidationRule> Rule = RuleInfo.RuleFactory.Execute();
+				Tasks.Emplace(Sequence, Rule);
+			}
+		}
+	}
+
+	const int32 ConfiguredMaxConcurrentTasks = CVarSequenceValidationMaxConcurrentTasks.GetValueOnAnyThread();
+	const int32 NumConcurrentTasks = ConfiguredMaxConcurrentTasks == -1 ? FPlatformMisc::NumberOfCores() : ConfiguredMaxConcurrentTasks;
+
+	NextTaskIndex = 0;
+	NumRemainingRunners = NumConcurrentTasks;
+
+	UE_LOG(LogSequenceValidator, Log, TEXT("Preparing %d sequence validator runners"), NumConcurrentTasks);
+	TArray<FSequenceValidatorTaskRunner> Runners;
+	Runners.Reserve(NumConcurrentTasks); // Important: we don't want runners to be re-allocated elsewhere once they start.
+	for (int32 RunnerIndex = 0; RunnerIndex < NumConcurrentTasks; ++RunnerIndex)
+	{
+		FSequenceValidatorTaskRunner& Runner = Runners.Emplace_GetRef(RunnerIndex);
+		Runner.Start(*this);
+	}
+
+	MainSignal->Wait();
+
+	UE_LOG(LogSequenceValidator, Log, TEXT("All sequence validator runners finished!"));
+	const bool bAllRunnersFinished = Algo::AllOf(Runners, 
+			[](const FSequenceValidatorTaskRunner& InRunner)
+			{
+				return !InRunner.IsRunning();
+			});
+	ensure(bAllRunnersFinished);
+
+	ensure(MainSignal);
+	FPlatformProcess::ReturnSynchEventToPool(MainSignal);
+	MainSignal = nullptr;
+
+	GatherResults(InQueue);
+
+	Tasks.Reset();
+
+	InOnFinished.ExecuteIfBound();
+}
+
+void FSequenceValidatorTaskScheduler::GatherResults(TArrayView<UMovieSceneSequence*> InQueue)
+{
+	Results.Reset();
+
+	// Create one root validation result for each root sequence. Add them to our results as we go.
+	TMap<UMovieSceneSequence*, TSharedPtr<FSequenceValidationResult>> SequenceToResult;
+	for (UMovieSceneSequence* Sequence : InQueue)
+	{
+		TSharedRef<FSequenceValidationResult> SequenceResult = MakeShared<FSequenceValidationResult>((UObject*)Sequence);
+		SequenceToResult.Add(Sequence, SequenceResult);
+		Results.AddResult(SequenceResult);
+	}
+
+	// Match validation results from the tasks to their owning sequence result.
+	for (FSequenceValidatorTask& Task : Tasks)
+	{
+		if (TSharedPtr<FSequenceValidationResult> SequenceResult = SequenceToResult.FindRef(Task.Sequence))
+		{
+			SequenceResult->AppendChildren(Task.TaskResults.GetResults());
+		}
+	}
+
+	// Add "empty" result for sequences with no messages.
+	for (UMovieSceneSequence* Sequence : InQueue)
+	{
+		TSharedPtr<FSequenceValidationResult> SequenceResult = SequenceToResult.FindRef(Sequence);
+		if (SequenceResult && !SequenceResult->HasChildren())
+		{
+			TSharedRef<FSequenceValidationResult> EmptyResult = MakeShared<FSequenceValidationResult>(
+					EMessageSeverity::Info, LOCTEXT("EmptyResultMessage", "No issues found."));
+			SequenceResult->AddChild(EmptyResult);
+		}
+	}
+}
+
+FSequenceValidatorTask* FSequenceValidatorTaskScheduler::AcquireNextTask(int32 InRunnerIndex)
+{
+	const int32 NewNextTaskIndex = FPlatformAtomics::InterlockedIncrement(&NextTaskIndex);
+	const int32 CurrentTaskIndex = NewNextTaskIndex - 1;
+	if (CurrentTaskIndex < Tasks.Num())
+	{
+		UE_LOG(LogSequenceValidator, Log, TEXT("Dispatching sequence validation task %d to runner %d"), CurrentTaskIndex, InRunnerIndex);
+		return &Tasks[CurrentTaskIndex];
+	}
+	return nullptr;
+}
+
+void FSequenceValidatorTaskScheduler::OnRunnerFinished(int32 InRunnerIndex)
+{
+	const int32 NewNumRemainingRunners = FPlatformAtomics::InterlockedDecrement(&NumRemainingRunners);
+	UE_LOG(LogSequenceValidator, Log, TEXT("Sequence validator runner %d finished, %d runners remaining"), InRunnerIndex, NewNumRemainingRunners);
+	if (NewNumRemainingRunners == 0)
+	{
+		MainSignal->Trigger();
+	}
+}
+
+}  // namespace Internal
+
+FSequenceValidator::FSequenceValidator()
+{
+	CreateValidationRuleInfos();
+}
+
+FSequenceValidator::~FSequenceValidator()
+{
+}
+
+TArrayView<UMovieSceneSequence* const> FSequenceValidator::GetQueue() const
+{
+	return ValidationQueue;
+}
+
+TConstArrayView<FSequenceValidationRuleInfo> FSequenceValidator::GetRules() const
+{
+	return ValidationRuleInfos;
+}
+
+const FSequenceValidationResults& FSequenceValidator::GetResults() const
+{
+	return ValidationResults;
+}
+
+void FSequenceValidator::Queue(UMovieSceneSequence* InSequence)
+{
+	ValidationQueue.AddUnique(InSequence);
+}
+
+void FSequenceValidator::Queue(TArrayView<UMovieSceneSequence*> InSequences)
+{
+	for (UMovieSceneSequence* Sequence : InSequences)
+	{
+		ValidationQueue.AddUnique(Sequence);
+	}
+}
+
+void FSequenceValidator::ClearQueue()
+{
+	ValidationQueue.Reset();
+}
+
+void FSequenceValidator::Validate(UMovieSceneSequence* InSequence)
+{
+	Queue(InSequence);
+	StartValidation(false);
+	ClearQueue();
+}
+
+void FSequenceValidator::Validate(TArrayView<UMovieSceneSequence*> InSequences)
+{
+	Queue(InSequences);
+	StartValidation(false);
+	ClearQueue();
+}
+
+void FSequenceValidator::StartValidation()
+{
+	StartValidation(true);
+}
+
+bool FSequenceValidator::IsValidating() const
+{
+	return Scheduler.IsValid();
+}
+
+FSimpleMulticastDelegate& FSequenceValidator::GetOnValidationFinished()
+{
+	return OnValidationFinishedDelegate;
+}
+
+void FSequenceValidator::CreateValidationRuleInfos()
+{
+	ISequenceValidatorModule& ValidationModule = FModuleManager::LoadModuleChecked<ISequenceValidatorModule>("SequenceValidator");
+	ValidationRuleInfos = ValidationModule.GetValidationRules();
+}
+
+void FSequenceValidator::StartValidation(bool bAsync)
+{
+	if (Scheduler.IsValid())
+	{
+		return;
+	}
+
+	ValidationResults.Reset();
+
+	if (ValidationQueue.IsEmpty())
+	{
+		return;
+	}
+
+	using namespace Internal;
+
+	Scheduler = MakeUnique<FSequenceValidatorTaskScheduler>(ValidationRuleInfos);
+	
+	if (bAsync)
+	{
+		UE::Tasks::Launch(TEXT("SequenceValidatorMain"),
+				[this]() 
+				{ 
+					Scheduler->RunQueue(
+							ValidationQueue,
+							FSimpleDelegate::CreateRaw(this, &FSequenceValidator::OnValidationFinished));
+				},
+				UE::Tasks::ETaskPriority::Normal);
+	}
+	else
+	{
+		Scheduler->RunQueue(
+				ValidationQueue,
+				FSimpleDelegate::CreateRaw(this, &FSequenceValidator::OnValidationFinished));
+	}
+}
+
+void FSequenceValidator::OnValidationFinished()
+{
+	ValidationResults = Scheduler->GetResults();
+	Scheduler.Reset();
+}
+
+}  // namespace UE::Sequencer
+
+#undef LOCTEXT_NAMESPACE
+
diff --git a/Engine/Plugins/MovieScene/SequenceValidator/Source/SequenceValidator/Private/Widgets/SSequenceValidator.cpp b/Engine/Plugins/MovieScene/SequenceValidator/Source/SequenceValidator/Private/Widgets/SSequenceValidator.cpp
new file mode 100644
--- /dev/null
+++ b/Engine/Plugins/MovieScene/SequenceValidator/Source/SequenceValidator/Private/Widgets/SSequenceValidator.cpp
@@ -0,0 +1,275 @@
+// Copyright Epic Games, Inc. All Rights Reserved.
+
+#include "Widgets/SSequenceValidator.h"
+
+#include "Framework/Application/SlateApplication.h"
+#include "Framework/Docking/TabManager.h"
+#include "Framework/Views/ITypedTableView.h"
+#include "IContentBrowserSingleton.h"
+#include "MovieSceneSequence.h"
+#include "SequenceValidatorCommands.h"
+#include "SequenceValidatorStyle.h"
+#include "ToolMenus.h"
+#include "Validation/SequenceValidator.h"
+#include "Widgets/Docking/SDockTab.h"
+#include "Widgets/Layout/SSplitter.h"
+#include "Widgets/SBoxPanel.h"
+#include "Widgets/SNullWidget.h"
+#include "Widgets/SSequenceValidatorQueue.h"
+#include "Widgets/SSequenceValidatorResults.h"
+#include "Widgets/SSequenceValidatorRules.h"
+#include "WorkspaceMenuStructure.h"
+#include "WorkspaceMenuStructureModule.h"
+
+#define LOCTEXT_NAMESPACE "SSequenceValidator"
+
+namespace UE::Sequencer
+{
+
+const FName SSequenceValidator::WindowName(TEXT("SequenceValidator"));
+const FName SSequenceValidator::MenubarName(TEXT("SequenceValidator.Menubar"));
+const FName SSequenceValidator::ToolbarName(TEXT("SequenceValidator.Toolbar"));
+
+void SSequenceValidator::RegisterTabSpawners()
+{
+	TSharedRef<FSequenceValidatorStyle> ValidatorStyle = FSequenceValidatorStyle::Get();
+
+	FGlobalTabmanager::Get()->RegisterNomadTabSpawner(
+		SSequenceValidator::WindowName,
+		FOnSpawnTab::CreateStatic(&SSequenceValidator::SpawnSequenceValidator)
+	)
+	.SetDisplayName(LOCTEXT("TabDisplayName", "Sequence Validator"))
+	.SetTooltipText(LOCTEXT("TabTooltipText", "Open the Sequence Validator tab."))
+	.SetIcon(FSlateIcon(ValidatorStyle->GetStyleSetName(), "SequenceValidator.TabIcon"))
+	.SetGroup(WorkspaceMenu::GetMenuStructure().GetLevelEditorCinematicsCategory())
+	.SetCanSidebarTab(false);
+}
+
+TSharedRef<SDockTab> SSequenceValidator::SpawnSequenceValidator(const FSpawnTabArgs& Args)
+{
+	auto NomadTab = SNew(SDockTab)
+		.TabRole(ETabRole::NomadTab)
+		.Label(LOCTEXT("TabTitle", "Sequence Validator"));
+
+	TSharedRef<SWidget> MainWidget = SNew(SSequenceValidator);
+	NomadTab->SetContent(MainWidget);
+	return NomadTab;
+}
+
+void SSequenceValidator::UnregisterTabSpawners()
+{
+	if (FSlateApplication::IsInitialized())
+	{
+		FGlobalTabmanager::Get()->UnregisterNomadTabSpawner(SSequenceValidator::WindowName);
+	}
+}
+
+void SSequenceValidator::Construct(const FArguments& InArgs)
+{
+	Validator = MakeShared<FSequenceValidator>();
+
+	// Setup commands.
+	const FSequenceValidatorCommands& Commands = FSequenceValidatorCommands::Get();
+	TSharedRef<FUICommandList> CommandList = MakeShareable(new FUICommandList);
+	CommandList->MapAction(
+			Commands.StartValidation,
+			FExecuteAction::CreateSP(this, &SSequenceValidator::StartValidation),
+			FCanExecuteAction::CreateSP(this, &SSequenceValidator::CanStartValidation));
+
+	// Build all UI elements.
+	TSharedRef<SWidget> MenubarContents = ConstructMenubar();
+	TSharedRef<SWidget> ToolbarContents = ConstructToolbar(CommandList);
+
+	// Main layout.
+	ChildSlot
+	[
+		SNew(SVerticalBox)
+		+SVerticalBox::Slot()
+		.AutoHeight()
+		[
+			MenubarContents
+		]
+		+SVerticalBox::Slot()
+		.AutoHeight()
+		[
+			ToolbarContents
+		]
+		+SVerticalBox::Slot()
+		.Padding(2.0)
+		.FillHeight(1.f)
+		[
+			SNew(SSplitter)
+			.Orientation(Orient_Horizontal)
+			+SSplitter::Slot()
+			[
+				SAssignNew(QueueWidget, SSequenceValidatorQueue)
+				.Validator(Validator)
+			]
+			+SSplitter::Slot()
+			[
+				SNew(SSplitter)
+				.Orientation(Orient_Vertical)
+				+SSplitter::Slot()
+				[
+					SAssignNew(RulesWidget, SSequenceValidatorRules)
+					.Validator(Validator)
+				]
+				+SSplitter::Slot()
+				[
+					SAssignNew(ResultsWidget, SSequenceValidatorResults)
+					.Validator(Validator)
+				]
+			]
+		]
+		+SVerticalBox::Slot()
+		.AutoHeight()
+		[
+			SNew(SHorizontalBox)
+			+SHorizontalBox::Slot()
+			[
+				SNew(STextBlock)
+				.Text(this, &SSequenceValidator::GetValidationStatusText)
+			]
+		]
+	];
+}
+
+TSharedRef<SWidget> SSequenceValidator::ConstructMenubar()
+{
+	UToolMenus* ToolMenus = UToolMenus::Get();
+	if (!ToolMenus->IsMenuRegistered(SSequenceValidator::MenubarName))
+	{
+		FToolMenuOwnerScoped ToolMenuOwnerScope(this);
+
+		UToolMenu* Menubar = ToolMenus->RegisterMenu(
+				SSequenceValidator::MenubarName, NAME_None, EMultiBoxType::MenuBar);
+	}
+
+	FToolMenuContext MenubarContext;
+	return ToolMenus->GenerateWidget(SSequenceValidator::MenubarName, MenubarContext);
+}
+
+TSharedRef<SWidget> SSequenceValidator::ConstructToolbar(TSharedRef<FUICommandList> InCommandList)
+{
+	UToolMenus* ToolMenus = UToolMenus::Get();
+	if (!ToolMenus->IsMenuRegistered(SSequenceValidator::ToolbarName))
+	{
+		FToolMenuOwnerScoped ToolMenuOwnerScope(this);
+
+		UToolMenu* Toolbar = ToolMenus->RegisterMenu(
+				SSequenceValidator::ToolbarName, NAME_None, EMultiBoxType::SlimHorizontalToolBar);
+
+		const FSequenceValidatorCommands& Commands = FSequenceValidatorCommands::Get();
+		TSharedRef<FSequenceValidatorStyle> ValidatorStyle = FSequenceValidatorStyle::Get();
+
+		FToolMenuSection& QueueSection = Toolbar->FindOrAddSection(TEXT("Queue"));
+		QueueSection.AddEntry(FToolMenuEntry::InitComboButton(
+			TEXT("Add"),
+			FUIAction(),
+			FNewToolMenuWidget::CreateStatic(&SSequenceValidator::GenerateAddSequenceMenu),
+			LOCTEXT("AddSequence_Label", "Add"),
+			LOCTEXT("AddSequence_ToolTip", "Add a sequence to the queue"),
+			FSlateIcon(ValidatorStyle->GetStyleSetName(), "SequenceValidator.QueueSequence")
+		));
+
+		FToolMenuSection& ValidationSection = Toolbar->FindOrAddSection(TEXT("Validation"));
+		ValidationSection.AddEntry(FToolMenuEntry::InitToolBarButton(Commands.StartValidation));
+	}
+
+	FToolMenuContext ToolbarContext(InCommandList, TSharedPtr<FExtender>());
+
+	USequenceValidatorMenuContext* Context = NewObject<USequenceValidatorMenuContext>();
+	Context->ValidatorWidget = SharedThis(this);
+	ToolbarContext.AddObject(Context);
+
+	return ToolMenus->GenerateWidget(SSequenceValidator::ToolbarName, ToolbarContext);
+}
+
+TSharedRef<SWidget> SSequenceValidator::GenerateAddSequenceMenu(const FToolMenuContext& Context)
+{
+	TSharedPtr<SSequenceValidator> This;
+	if (USequenceValidatorMenuContext* ValidatorContext = Context.FindContext<USequenceValidatorMenuContext>())
+	{
+		This = ValidatorContext->ValidatorWidget.Pin();
+	}
+
+	if (!This)
+	{
+		return SNullWidget::NullWidget;
+	}
+
+	FAssetPickerConfig PickerConfig;
+	PickerConfig.bAllowNullSelection = false;
+	PickerConfig.bFocusSearchBoxWhenOpened = true;
+	PickerConfig.Filter.bRecursiveClasses = true;
+	PickerConfig.Filter.ClassPaths.Add(UMovieSceneSequence::StaticClass()->GetClassPathName());
+	PickerConfig.SelectionMode = ESelectionMode::Single;
+	PickerConfig.InitialAssetViewType = EAssetViewType::List;
+	PickerConfig.OnAssetSelected = FOnAssetSelected::CreateSP(This.Get(), &SSequenceValidator::OnAddSequence);
+	PickerConfig.OnAssetsActivated = FOnAssetsActivated::CreateSP(This.Get(), &SSequenceValidator::OnAddSequences);
+
+	return IContentBrowserSingleton::Get().CreateAssetPicker(PickerConfig);
+}
+
+void SSequenceValidator::OnAddSequence(const FAssetData& SelectedAsset)
+{
+	if (UMovieSceneSequence* Sequence = Cast<UMovieSceneSequence>(SelectedAsset.GetAsset()))
+	{
+		Validator->Queue(Sequence);
+	}
+
+	QueueWidget->RequestListRefresh();
+}
+
+void SSequenceValidator::OnAddSequences(const TArray<FAssetData>& ActivatedAssets, EAssetTypeActivationMethod::Type ActivationMethod)
+{
+	for (const FAssetData& AssetData : ActivatedAssets)
+	{
+		if (UMovieSceneSequence* Sequence = Cast<UMovieSceneSequence>(AssetData.GetAsset()))
+		{
+			Validator->Queue(Sequence);
+		}
+	}
+
+	QueueWidget->RequestListRefresh();
+}
+
+void SSequenceValidator::Tick(const FGeometry& AllottedGeometry, const double InCurrentTime, const float InDeltaTime)
+{
+	SCompoundWidget::Tick(AllottedGeometry, InCurrentTime, InDeltaTime);
+
+	if (bWaitingForValidationToFinish && !Validator->IsValidating())
+	{
+		bWaitingForValidationToFinish = false;
+
+		ResultsWidget->RequestListRefresh();
+	}
+}
+
+void SSequenceValidator::StartValidation()
+{
+	bWaitingForValidationToFinish = true;
+	Validator->StartValidation();
+}
+
+bool SSequenceValidator::CanStartValidation() const
+{
+	return !Validator->IsValidating();
+}
+
+EVisibility SSequenceValidator::GetValidationStatusVisibility() const
+{
+	return Validator->IsValidating() ? EVisibility::Visible : EVisibility::Collapsed;
+}
+
+FText SSequenceValidator::GetValidationStatusText() const
+{
+	return Validator->IsValidating() ? 
+		FText::Format(LOCTEXT("ValidationStatus_Running", "Validating {0} sequences..."), Validator->GetQueue().Num()) :
+		LOCTEXT("ValidationStatus_Waiting", "Ready to validate");
+}
+
+}  // namespace UE::Sequencer
+
+#undef LOCTEXT_NAMESPACE
+
diff --git a/Engine/Plugins/MovieScene/SequenceValidator/Source/SequenceValidator/Private/Widgets/SSequenceValidator.h b/Engine/Plugins/MovieScene/SequenceValidator/Source/SequenceValidator/Private/Widgets/SSequenceValidator.h
new file mode 100644
--- /dev/null
+++ b/Engine/Plugins/MovieScene/SequenceValidator/Source/SequenceValidator/Private/Widgets/SSequenceValidator.h
@@ -0,0 +1,86 @@
+// Copyright Epic Games, Inc. All Rights Reserved.
+
+#pragma once
+
+#include "ContentBrowserDelegates.h"
+#include "Widgets/SCompoundWidget.h"
+
+#include "SSequenceValidator.generated.h"
+
+class FSpawnTabArgs;
+class FUICommandList;
+class SDockTab;
+struct FAssetData;
+struct FToolMenuContext;
+
+namespace UE::Sequencer
+{
+
+class FSequenceValidator;
+class SSequenceValidatorQueue;
+class SSequenceValidatorResults;
+class SSequenceValidatorRules;
+
+class SSequenceValidator : public SCompoundWidget
+{
+public:
+
+	static const FName WindowName;
+	static const FName MenubarName;
+	static const FName ToolbarName;
+
+	static void RegisterTabSpawners();
+	static TSharedRef<SDockTab> SpawnSequenceValidator(const FSpawnTabArgs& Args);
+	static void UnregisterTabSpawners();
+
+public:
+
+	SLATE_BEGIN_ARGS(SSequenceValidator)
+	{}
+	SLATE_END_ARGS()
+
+	void Construct(const FArguments& InArgs);
+
+protected:
+
+	// SWidget interface.
+	virtual void Tick(const FGeometry& AllottedGeometry, const double InCurrentTime, const float InDeltaTime) override;
+
+private:
+
+	TSharedRef<SWidget> ConstructMenubar();
+	TSharedRef<SWidget> ConstructToolbar(TSharedRef<FUICommandList> InCommandList);
+
+	static TSharedRef<SWidget> GenerateAddSequenceMenu(const FToolMenuContext& Context);
+	void OnAddSequence(const FAssetData& SelectedAsset);
+	void OnAddSequences(const TArray<FAssetData>& ActivatedAssets, EAssetTypeActivationMethod::Type ActivationMethod);
+
+	void StartValidation();
+	bool CanStartValidation() const;
+
+	EVisibility GetValidationStatusVisibility() const;
+	FText GetValidationStatusText() const;
+
+private:
+
+	TSharedPtr<SSequenceValidatorQueue> QueueWidget;
+	TSharedPtr<SSequenceValidatorResults> ResultsWidget;
+	TSharedPtr<SSequenceValidatorRules> RulesWidget;
+
+	TSharedPtr<FSequenceValidator> Validator;
+
+	bool bWaitingForValidationToFinish = false;
+};
+
+}  // namespace UE::Sequencer
+
+UCLASS()
+class USequenceValidatorMenuContext : public UObject
+{
+	GENERATED_BODY()
+
+public:
+
+	TWeakPtr<UE::Sequencer::SSequenceValidator> ValidatorWidget;
+};
+
diff --git a/Engine/Plugins/MovieScene/SequenceValidator/Source/SequenceValidator/Private/Widgets/SSequenceValidatorQueue.cpp b/Engine/Plugins/MovieScene/SequenceValidator/Source/SequenceValidator/Private/Widgets/SSequenceValidatorQueue.cpp
new file mode 100644
--- /dev/null
+++ b/Engine/Plugins/MovieScene/SequenceValidator/Source/SequenceValidator/Private/Widgets/SSequenceValidatorQueue.cpp
@@ -0,0 +1,223 @@
+// Copyright Epic Games, Inc. All Rights Reserved.
+
+#include "Widgets/SSequenceValidatorQueue.h"
+
+#include "AssetThumbnail.h"
+#include "DragAndDrop/AssetDragDropOp.h"
+#include "SequenceValidatorStyle.h"
+#include "ThumbnailRendering/ThumbnailManager.h"
+#include "Validation/SequenceValidator.h"
+#include "Widgets/SBoxPanel.h"
+
+#define LOCTEXT_NAMESPACE "SSequenceValidatorQueue"
+
+namespace UE::Sequencer
+{
+
+class SSequenceValidatorQueueEntry : public STableRow<UMovieSceneSequence*>
+{
+public:
+
+	SLATE_BEGIN_ARGS(SSequenceValidatorQueueEntry)
+	{}
+		SLATE_ARGUMENT(TWeakObjectPtr<UMovieSceneSequence>, Item)
+	SLATE_END_ARGS()
+
+	void Construct(const FArguments& InArgs, const TSharedRef<STableViewBase>& OwnerTable);
+
+private:
+
+	const FSlateBrush* GetThumbnailBorder() const;
+	FText GetSequenceName() const;
+
+private:
+
+	TWeakObjectPtr<UMovieSceneSequence> WeakItem;
+
+	TSharedPtr<FAssetThumbnail> AssetThumbnail;
+	TSharedPtr<SBorder> ThumbnailBorder;
+};
+
+void SSequenceValidatorQueueEntry::Construct(const FArguments& InArgs, const TSharedRef<STableViewBase>& OwnerTable)
+{
+	using FSuperRowType = STableRow<UMovieSceneSequence*>;
+
+	WeakItem = InArgs._Item;
+
+	FIntPoint ThumbnailSize(48, 48);
+	TSharedPtr<FAssetThumbnailPool> ThumbnailPool = UThumbnailManager::Get().GetSharedThumbnailPool();
+
+	FAssetData AssetData(WeakItem.Get());
+	AssetThumbnail = MakeShared<FAssetThumbnail>(AssetData, ThumbnailSize.X, ThumbnailSize.Y, ThumbnailPool);
+	FAssetThumbnailConfig AssetThumbnailConfig;
+
+	ChildSlot
+	.Padding(0, 3, 5, 0)
+	.VAlign(VAlign_Center)
+	[
+		SNew(SHorizontalBox)
+		+SHorizontalBox::Slot()
+		.AutoWidth()
+		[
+			SNew(SBorder)
+			.Padding(0, 0, 4, 4)
+			.Visibility(EVisibility::SelfHitTestInvisible)
+			.BorderImage(FAppStyle::Get().GetBrush("PropertyEditor.AssetTileItem.DropShadow"))
+			[
+				SNew(SOverlay)
+				+SOverlay::Slot()
+				.Padding(1.0f)
+				[
+					SAssignNew(ThumbnailBorder, SBorder)
+					.Padding(0)
+					.BorderImage(FStyleDefaults::GetNoBrush())
+					[
+						SNew(SBox)
+						.WidthOverride(ThumbnailSize.X)
+						.HeightOverride(ThumbnailSize.Y)
+						[
+							AssetThumbnail->MakeThumbnailWidget(AssetThumbnailConfig)
+						]
+					]
+				]
+				+SOverlay::Slot()
+				[
+					SNew(SImage)
+					.Image(this, &SSequenceValidatorQueueEntry::GetThumbnailBorder)
+					.Visibility(EVisibility::SelfHitTestInvisible)
+				]
+			]
+		]
+		+SHorizontalBox::Slot()
+		.Padding(8, 4, 4, 4)
+		.FillWidth(1.f)
+		[
+			SNew(STextBlock)
+			.Text(this, &SSequenceValidatorQueueEntry::GetSequenceName)
+		]
+	];
+
+	FSuperRowType::ConstructInternal(
+		FSuperRowType::FArguments(),
+		OwnerTable);
+}
+
+const FSlateBrush* SSequenceValidatorQueueEntry::GetThumbnailBorder() const
+{
+	static const FName HoveredBorderName("PropertyEditor.AssetThumbnailBorderHovered");
+	static const FName RegularBorderName("PropertyEditor.AssetThumbnailBorder");
+
+	return ThumbnailBorder->IsHovered() ? FAppStyle::Get().GetBrush(HoveredBorderName) : FAppStyle::Get().GetBrush(RegularBorderName);
+}
+
+FText SSequenceValidatorQueueEntry::GetSequenceName() const
+{
+	if (UMovieSceneSequence* Item = WeakItem.Get())
+	{
+		return Item->GetDisplayName();
+	}
+	return FText::GetEmpty();
+}
+
+void SSequenceValidatorQueue::Construct(const FArguments& InArgs)
+{
+	Validator = InArgs._Validator;
+
+	TSharedRef<FSequenceValidatorStyle> ValidatorStyle = FSequenceValidatorStyle::Get();
+
+	ChildSlot
+	[
+		SNew(SVerticalBox)
+		+SVerticalBox::Slot()
+		.AutoHeight()
+		[
+			SNew(STextBlock)
+			.Margin(4.f)
+			.Text(LOCTEXT("QueueTitle", "Queue"))
+			.TextStyle(ValidatorStyle, "SequenceValidator.PanelTitleText")
+		]
+		+SVerticalBox::Slot()
+		.FillHeight(1.f)
+		[
+			SAssignNew(ListView, SListView<UMovieSceneSequence*>)
+			.ListItemsSource(&ItemSource)
+			.OnGenerateRow(this, &SSequenceValidatorQueue::OnListGenerateItemRow)
+		]
+	];
+}
+
+TSharedRef<ITableRow> SSequenceValidatorQueue::OnListGenerateItemRow(UMovieSceneSequence* Item, const TSharedRef<STableViewBase>& OwnerTable)
+{
+	return SNew(SSequenceValidatorQueueEntry, OwnerTable)
+		.Item(Item);
+}
+
+void SSequenceValidatorQueue::Tick(const FGeometry& AllottedGeometry, const double InCurrentTime, const float InDeltaTime)
+{
+	SCompoundWidget::Tick(AllottedGeometry, InCurrentTime, InDeltaTime);
+	
+	if (bUpdateItemSource)
+	{
+		bUpdateItemSource = false;
+
+		UpdateItemSource();
+
+		ListView->RequestListRefresh();
+	}
+}
+
+void SSequenceValidatorQueue::OnDragEnter(const FGeometry& MyGeometry, const FDragDropEvent& DragDropEvent)
+{
+	SCompoundWidget::OnDragEnter(MyGeometry, DragDropEvent);
+}
+
+FReply SSequenceValidatorQueue::OnDrop(const FGeometry& MyGeometry, const FDragDropEvent& DragDropEvent)
+{
+	TArray<UMovieSceneSequence*> DroppedSequences;
+
+	if (TSharedPtr<FAssetDragDropOp> AssetDragDropOp = DragDropEvent.GetOperationAs<FAssetDragDropOp>())
+	{
+		const TArray<FAssetData>& Assets = AssetDragDropOp->GetAssets();
+		for (const FAssetData& Asset : Assets)
+		{
+			const UClass* AssetClass = Asset.GetClass();
+			if (AssetClass && AssetClass->IsChildOf<UMovieSceneSequence>())
+			{
+				if (UMovieSceneSequence* Sequence = Cast<UMovieSceneSequence>(Asset.GetAsset()))
+				{
+					DroppedSequences.Add(Sequence);
+				}
+			}
+		}
+	}
+
+	if (!DroppedSequences.IsEmpty())
+	{
+		Validator->Queue(DroppedSequences);
+		RequestListRefresh();
+
+		return FReply::Handled();
+	}
+
+	return SCompoundWidget::OnDrop(MyGeometry, DragDropEvent);
+}
+
+void SSequenceValidatorQueue::RequestListRefresh()
+{
+	bUpdateItemSource = true;
+}
+
+void SSequenceValidatorQueue::UpdateItemSource()
+{
+	ItemSource.Reset();
+
+	for (UMovieSceneSequence* Sequence : Validator->GetQueue())
+	{
+		ItemSource.Add(Sequence);
+	}
+}
+
+}  // namespace UE::Sequencer
+
+#undef LOCTEXT_NAMESPACE
+
diff --git a/Engine/Plugins/MovieScene/SequenceValidator/Source/SequenceValidator/Private/Widgets/SSequenceValidatorQueue.h b/Engine/Plugins/MovieScene/SequenceValidator/Source/SequenceValidator/Private/Widgets/SSequenceValidatorQueue.h
new file mode 100644
--- /dev/null
+++ b/Engine/Plugins/MovieScene/SequenceValidator/Source/SequenceValidator/Private/Widgets/SSequenceValidatorQueue.h
@@ -0,0 +1,51 @@
+// Copyright Epic Games, Inc. All Rights Reserved.
+
+#pragma once
+
+#include "MovieSceneSequence.h"
+#include "Widgets/SCompoundWidget.h"
+#include "Widgets/Views/SListView.h"
+
+namespace UE::Sequencer
+{
+
+class FSequenceValidator;
+
+class SSequenceValidatorQueue : public SCompoundWidget
+{
+public:
+
+	SLATE_BEGIN_ARGS(SSequenceValidatorQueue)
+	{}
+		SLATE_ARGUMENT(TSharedPtr<FSequenceValidator>, Validator)
+	SLATE_END_ARGS()
+
+	void Construct(const FArguments& InArgs);
+
+	void RequestListRefresh();
+
+protected:
+
+	// SWidget interface.
+	virtual void Tick(const FGeometry& AllottedGeometry, const double InCurrentTime, const float InDeltaTime) override;
+	virtual void OnDragEnter(const FGeometry& MyGeometry, const FDragDropEvent& DragDropEvent) override;
+	virtual FReply OnDrop(const FGeometry& MyGeometry, const FDragDropEvent& DragDropEvent) override;
+
+private:
+
+	TSharedRef<ITableRow> OnListGenerateItemRow(UMovieSceneSequence* Item, const TSharedRef<STableViewBase>& OwnerTable);
+
+	void UpdateItemSource();
+
+private:
+
+	TSharedPtr<FSequenceValidator> Validator;
+
+	TSharedPtr<SListView<UMovieSceneSequence*>> ListView;
+	TArray<UMovieSceneSequence*> ItemSource;
+
+	bool bUpdateItemSource = false;
+};
+
+}  // namespace UE::Sequencer
+
diff --git a/Engine/Plugins/MovieScene/SequenceValidator/Source/SequenceValidator/Private/Widgets/SSequenceValidatorResults.cpp b/Engine/Plugins/MovieScene/SequenceValidator/Source/SequenceValidator/Private/Widgets/SSequenceValidatorResults.cpp
new file mode 100644
--- /dev/null
+++ b/Engine/Plugins/MovieScene/SequenceValidator/Source/SequenceValidator/Private/Widgets/SSequenceValidatorResults.cpp
@@ -0,0 +1,290 @@
+// Copyright Epic Games, Inc. All Rights Reserved.
+
+#include "Widgets/SSequenceValidatorResults.h"
+
+#include "Editor.h"
+#include "ILevelSequenceEditorToolkit.h"
+#include "ISequencer.h"
+#include "LevelSequence.h"
+#include "Misc/ScopedSlowTask.h"
+#include "MovieSceneSection.h"
+#include "MovieSceneSequence.h"
+#include "MovieSceneSequenceID.h"
+#include "MovieSceneTrack.h"
+#include "SequenceValidatorStyle.h"
+#include "Subsystems/AssetEditorSubsystem.h"
+#include "Validation/SequenceValidationResult.h"
+#include "Validation/SequenceValidator.h"
+
+#define LOCTEXT_NAMESPACE "SSequenceValidatorResults"
+
+namespace UE::Sequencer
+{
+
+class SSequenceValidatorResultEntry : public SCompoundWidget
+{
+public:
+
+	SLATE_BEGIN_ARGS(SSequenceValidatorResultEntry)
+	{}
+		SLATE_ARGUMENT(TSharedPtr<FSequenceValidationResult>, Item)
+	SLATE_END_ARGS()
+
+	void Construct(const FArguments& InArgs);
+
+private:
+
+	const FSlateBrush* GetIcon() const;
+	FText GetMessage() const;
+
+private:
+
+	TWeakPtr<FSequenceValidationResult> WeakItem;
+};
+
+void SSequenceValidatorResultEntry::Construct(const FArguments& InArgs)
+{
+	WeakItem = InArgs._Item;
+
+	ChildSlot
+	.Padding(0, 3, 5, 0)
+	.VAlign(VAlign_Center)
+	[
+		SNew(SHorizontalBox)
+		+SHorizontalBox::Slot()
+		.AutoWidth()
+		.VAlign(VAlign_Center)
+		[
+			SNew(SImage)
+			.Image(this, &SSequenceValidatorResultEntry::GetIcon)
+			.ColorAndOpacity(FSlateColor::UseForeground())
+		]
+		+SHorizontalBox::Slot()
+		.AutoWidth()
+		.VAlign(VAlign_Center)
+		.Padding(2, 0)
+		[
+			SNew(STextBlock)
+			.Text(this, &SSequenceValidatorResultEntry::GetMessage)
+		]
+	];
+}
+
+const FSlateBrush* SSequenceValidatorResultEntry::GetIcon() const
+{
+	if (TSharedPtr<FSequenceValidationResult> Item = WeakItem.Pin())
+	{
+		if (UObject* Target = Item->GetTarget())
+		{
+			TSharedRef<FSequenceValidatorStyle> ValidatorStyle = FSequenceValidatorStyle::Get();
+
+			if (Target->IsA<UMovieSceneSection>())
+			{
+				return ValidatorStyle->GetBrush("ValidationResult.Section");
+			}
+			else if (Target->IsA<UMovieSceneTrack>())
+			{
+				return ValidatorStyle->GetBrush("ValidationResult.Track");
+			}
+			else if (Target->IsA<UMovieSceneSequence>())
+			{
+				return ValidatorStyle->GetBrush("ValidationResult.Sequence");
+			}
+		}
+	}
+	return nullptr;
+}
+
+FText SSequenceValidatorResultEntry::GetMessage() const
+{
+	if (TSharedPtr<FSequenceValidationResult> Item = WeakItem.Pin())
+	{
+		const FText& UserMessage = Item->GetUserMessage();
+		if (!UserMessage.IsEmpty())
+		{
+			return UserMessage;
+		}
+
+		UObject* Target = Item->GetTarget();
+		if (Target)
+		{
+			return FText::FromString(GetNameSafe(Target));
+		}
+	}
+	return FText::GetEmpty();
+}
+
+void SSequenceValidatorResults::Construct(const FArguments& InArgs)
+{
+	Validator = InArgs._Validator;
+
+	TSharedRef<FSequenceValidatorStyle> ValidatorStyle = FSequenceValidatorStyle::Get();
+
+	ChildSlot
+	[
+		SNew(SVerticalBox)
+		+SVerticalBox::Slot()
+		.AutoHeight()
+		[
+			SNew(STextBlock)
+			.Margin(4.f)
+			.Text(LOCTEXT("ResultsTitle", "Results"))
+			.TextStyle(ValidatorStyle, "SequenceValidator.PanelTitleText")
+		]
+		+SVerticalBox::Slot()
+		.FillHeight(1.f)
+		[
+			SNew(SOverlay)
+			+SOverlay::Slot()
+			[
+				SNew(SBorder)
+				.HAlign(HAlign_Center)
+				.VAlign(VAlign_Center)
+				.Visibility(this, &SSequenceValidatorResults::GetEmptyResultsMessageVisibility)
+				[
+					SNew(STextBlock)
+					.Text(LOCTEXT("EmptyResultsMessage", "All good!"))
+				]
+			]
+			+SOverlay::Slot()
+			[
+				SAssignNew(TreeView, STreeView<FTreeItemPtr>)
+				.TreeItemsSource(&ItemSource)
+				.OnGenerateRow(this, &SSequenceValidatorResults::OnTreeViewGenerateRow)
+				.OnGetChildren(this, &SSequenceValidatorResults::OnTreeViewGetChildren)
+				.OnMouseButtonDoubleClick(this, &SSequenceValidatorResults::OnTreeViewMouseButtonDoubleClick)
+			]
+		]
+	];
+}
+
+TSharedRef<ITableRow> SSequenceValidatorResults::OnTreeViewGenerateRow(TSharedPtr<FSequenceValidationResult> Item, const TSharedRef<STableViewBase>& OwnerTable)
+{
+	TSharedRef<FSequenceValidatorStyle> ValidatorStyle = FSequenceValidatorStyle::Get();
+	return SNew(STableRow<TSharedPtr<FSequenceValidationResult>>, OwnerTable)
+			.Style(&ValidatorStyle->GetWidgetStyle<FTableRowStyle>("ValidationResult.RowStyle"))
+			[
+				SNew(SSequenceValidatorResultEntry)
+					.Item(Item)
+			];
+}
+
+void SSequenceValidatorResults::OnTreeViewGetChildren(FTreeItemPtr Item, TArray<FTreeItemPtr>& OutChildren)
+{
+	OutChildren.Append(Item->GetChildren());
+}
+
+void SSequenceValidatorResults::OnTreeViewMouseButtonDoubleClick(FTreeItemPtr Item)
+{
+	TSharedPtr<FSequenceValidationResult> RootItem = Item->GetRoot();
+	if (!RootItem)
+	{
+		return;
+	}
+
+	UObject* Asset = RootItem->GetTarget();
+	if (!Asset)
+	{
+		return;
+	}
+	
+	// Open the root sequence for this result, and grab the Sequencer editor that should now be open.
+	UAssetEditorSubsystem* AssetEditorSubsystem = GEditor->GetEditorSubsystem<UAssetEditorSubsystem>();
+	{
+		FScopedSlowTask SlowTask(1.f, LOCTEXT("OpenSequenceSlowTask", "Loading sequence..."));
+		SlowTask.MakeDialog();
+		SlowTask.EnterProgressFrame();
+
+		AssetEditorSubsystem->OpenEditorForAsset(Asset);
+	}
+
+	TSharedPtr<ISequencer> Sequencer;
+	if (Asset->IsA<ULevelSequence>())
+	{
+		ILevelSequenceEditorToolkit* SequencerEditor = (ILevelSequenceEditorToolkit*)AssetEditorSubsystem->FindEditorForAsset(Asset, true);
+		if (SequencerEditor)
+		{
+			Sequencer = SequencerEditor->GetSequencer();
+		}
+	}
+
+	if (!ensureMsgf(Sequencer, TEXT("Couldn't open Sequencer editor for root sequence '%s'"), *GetNameSafe(Asset)))
+	{
+		return;
+	}
+
+	UObject* Target = Item->GetTarget();
+	if (Target)
+	{
+		// Focus the sub-sequence for the target (it might be a sub-sequence, or the currently open Sequencer
+		// may be focused on something else).
+		Sequencer->PopToSequenceInstance(MovieSceneSequenceID::Root);
+		TArray<UMovieSceneSubSection*> SubSectionTrail;
+		Item->GetSubSectionTrail(SubSectionTrail);
+		for (UMovieSceneSubSection* SubSection : SubSectionTrail)
+		{
+			Sequencer->FocusSequenceInstance(*SubSection);
+		}
+
+		// Select the track/section/etc that relates to the message.
+		if (UMovieSceneTrack* TargetTrack = Cast<UMovieSceneTrack>(Target))
+		{
+			Sequencer->SelectTrack(TargetTrack);
+		}
+		else if (UMovieSceneSection* TargetSection = Cast<UMovieSceneSection>(Target))
+		{
+			Sequencer->SelectSection(TargetSection);
+		}
+
+		// Move the time head to a time that relates to the message.
+		if (Item->HasLocalTime())
+		{
+			Sequencer->SetLocalTime(Item->GetLocalTime());
+		}
+	}
+}
+
+EVisibility SSequenceValidatorResults::GetEmptyResultsMessageVisibility() const
+{
+	return ItemSource.IsEmpty() ? EVisibility::Visible : EVisibility::Hidden;
+}
+
+void SSequenceValidatorResults::Tick(const FGeometry& AllottedGeometry, const double InCurrentTime, const float InDeltaTime)
+{
+	SCompoundWidget::Tick(AllottedGeometry, InCurrentTime, InDeltaTime);
+
+	if (bUpdateItemSource)
+	{
+		bUpdateItemSource = false;
+
+		UpdateItemSource();
+
+		TreeView->RequestTreeRefresh();
+
+		for (FTreeItemPtr Item : ItemSource)
+		{
+			TreeView->SetItemExpansion(Item, true);
+		}
+	}
+}
+
+void SSequenceValidatorResults::RequestListRefresh()
+{
+	bUpdateItemSource = true;
+}
+
+void SSequenceValidatorResults::UpdateItemSource()
+{
+	ItemSource.Reset();
+
+	const FSequenceValidationResults& Results = Validator->GetResults();
+	for (TSharedPtr<FSequenceValidationResult> Result : Results.GetResults())
+	{
+		ItemSource.Add(Result);
+	}
+}
+
+}  // namespace UE::Sequencer
+
+#undef LOCTEXT_NAMESPACE
+
diff --git a/Engine/Plugins/MovieScene/SequenceValidator/Source/SequenceValidator/Private/Widgets/SSequenceValidatorResults.h b/Engine/Plugins/MovieScene/SequenceValidator/Source/SequenceValidator/Private/Widgets/SSequenceValidatorResults.h
new file mode 100644
--- /dev/null
+++ b/Engine/Plugins/MovieScene/SequenceValidator/Source/SequenceValidator/Private/Widgets/SSequenceValidatorResults.h
@@ -0,0 +1,55 @@
+// Copyright Epic Games, Inc. All Rights Reserved.
+
+#pragma once
+
+#include "Widgets/SCompoundWidget.h"
+#include "Widgets/Views/STreeView.h"
+
+namespace UE::Sequencer
+{
+
+class FSequenceValidator;
+class FSequenceValidationResult;
+
+class SSequenceValidatorResults : public SCompoundWidget
+{
+public:
+
+	SLATE_BEGIN_ARGS(SSequenceValidatorResults)
+	{}
+		SLATE_ARGUMENT(TSharedPtr<FSequenceValidator>, Validator)
+	SLATE_END_ARGS()
+
+	void Construct(const FArguments& InArgs);
+
+	void RequestListRefresh();
+
+protected:
+
+	// SWidget interface.
+	virtual void Tick(const FGeometry& AllottedGeometry, const double InCurrentTime, const float InDeltaTime) override;
+
+private:
+
+	using FTreeItemPtr = TSharedPtr<FSequenceValidationResult>;
+
+	TSharedRef<ITableRow> OnTreeViewGenerateRow(FTreeItemPtr Item, const TSharedRef<STableViewBase>& OwnerTable);
+	void OnTreeViewGetChildren(FTreeItemPtr Item, TArray<FTreeItemPtr>& OutChildren);
+	void OnTreeViewMouseButtonDoubleClick(FTreeItemPtr Item);
+
+	void UpdateItemSource();
+	
+	EVisibility GetEmptyResultsMessageVisibility() const;
+
+private:
+
+	TSharedPtr<FSequenceValidator> Validator;
+
+	TSharedPtr<STreeView<FTreeItemPtr>> TreeView;
+	TArray<FTreeItemPtr> ItemSource;
+
+	bool bUpdateItemSource = false;
+};
+
+}  // namespace UE::Sequencer
+
diff --git a/Engine/Plugins/MovieScene/SequenceValidator/Source/SequenceValidator/Private/Widgets/SSequenceValidatorRules.cpp b/Engine/Plugins/MovieScene/SequenceValidator/Source/SequenceValidator/Private/Widgets/SSequenceValidatorRules.cpp
new file mode 100644
--- /dev/null
+++ b/Engine/Plugins/MovieScene/SequenceValidator/Source/SequenceValidator/Private/Widgets/SSequenceValidatorRules.cpp
@@ -0,0 +1,176 @@
+// Copyright Epic Games, Inc. All Rights Reserved.
+
+#include "Widgets/SSequenceValidatorRules.h"
+
+#include "SequenceValidatorStyle.h"
+#include "Validation/SequenceValidationRule.h"
+#include "Validation/SequenceValidator.h"
+#include "Widgets/Input/SCheckBox.h"
+
+#define LOCTEXT_NAMESPACE "SSequenceValidatorRules"
+
+namespace UE::Sequencer
+{
+
+class SSequenceValidatorRuleEntry : public STableRow<TSharedPtr<FSequenceValidationRuleInfo>>
+{
+public:
+
+	SLATE_BEGIN_ARGS(SSequenceValidatorRuleEntry)
+	{}
+		SLATE_ARGUMENT(TWeakPtr<FSequenceValidationRuleInfo>, Item)
+	SLATE_END_ARGS()
+
+	void Construct(const FArguments& InArgs, const TSharedRef<STableViewBase>& OwnerTable);
+
+private:
+
+	FText GetRuleName() const;
+	FText GetRuleDescription() const;
+	ECheckBoxState IsRuleEnabled() const;
+	void OnRuleEnabledChanged(ECheckBoxState CheckState) const;
+
+private:
+
+	TWeakPtr<FSequenceValidationRuleInfo> WeakItem;
+};
+
+void SSequenceValidatorRuleEntry::Construct(const FArguments& InArgs, const TSharedRef<STableViewBase>& OwnerTable)
+{
+	using FSuperRowType = STableRow<TSharedPtr<FSequenceValidationRuleInfo>>;
+
+	WeakItem = InArgs._Item;
+
+	const ISlateStyle& AppStyle = FAppStyle::Get();
+	TSharedRef<const FSequenceValidatorStyle> ValidatorStyle = FSequenceValidatorStyle::Get();
+
+	ChildSlot
+	.Padding(4.f)
+	[
+		SNew(SHorizontalBox)
+		+SHorizontalBox::Slot()
+		.AutoWidth()
+		.VAlign(VAlign_Center)
+		[
+			SNew(SImage)
+			.Image(ValidatorStyle->GetBrush("ValidationRule"))
+		]
+		+SHorizontalBox::Slot()
+		.Padding(8.f, 4.f)
+		.FillWidth(1.f)
+		[
+			SNew(SVerticalBox)
+			+SVerticalBox::Slot()
+			[
+				SNew(STextBlock)
+				.Text(this, &SSequenceValidatorRuleEntry::GetRuleName)
+				.TextStyle(ValidatorStyle, "ValidatorRule.RuleNameText")
+				.ToolTipText(this, &SSequenceValidatorRuleEntry::GetRuleDescription)
+			]
+			+SVerticalBox::Slot()
+			[
+				SNew(STextBlock)
+				.Text(this, &SSequenceValidatorRuleEntry::GetRuleDescription)
+				.TextStyle(ValidatorStyle, "ValidatorRule.RuleDescriptionText")
+			]
+		]
+		+SHorizontalBox::Slot()
+		.AutoWidth()
+		[
+			SNew(SCheckBox)
+			.IsChecked(this, &SSequenceValidatorRuleEntry::IsRuleEnabled)
+			.OnCheckStateChanged(this, &SSequenceValidatorRuleEntry::OnRuleEnabledChanged)
+			.ToolTipText(LOCTEXT("ToggleValidationRule", "Toggle this validation rule"))
+		]
+	];
+
+	FSuperRowType::ConstructInternal(
+		FSuperRowType::FArguments()
+			.Style(&ValidatorStyle->GetWidgetStyle<FTableRowStyle>("ValidationRule.RowStyle")),
+		OwnerTable);
+}
+
+FText SSequenceValidatorRuleEntry::GetRuleName() const
+{
+	if (TSharedPtr<FSequenceValidationRuleInfo> Item = WeakItem.Pin())
+	{
+		return Item->RuleName;
+	}
+	return FText::GetEmpty();
+}
+
+FText SSequenceValidatorRuleEntry::GetRuleDescription() const
+{
+	if (TSharedPtr<FSequenceValidationRuleInfo> Item = WeakItem.Pin())
+	{
+		return Item->RuleDescription;
+	}
+	return FText::GetEmpty();
+}
+
+ECheckBoxState SSequenceValidatorRuleEntry::IsRuleEnabled() const
+{
+	if (TSharedPtr<FSequenceValidationRuleInfo> Item = WeakItem.Pin())
+	{
+		return Item->bIsEnabled ? ECheckBoxState::Checked : ECheckBoxState::Unchecked;
+	}
+	return ECheckBoxState::Undetermined;
+}
+
+void SSequenceValidatorRuleEntry::OnRuleEnabledChanged(ECheckBoxState CheckState) const
+{
+	if (TSharedPtr<FSequenceValidationRuleInfo> Item = WeakItem.Pin())
+	{
+		Item->bIsEnabled = (CheckState == ECheckBoxState::Checked);
+	}
+}
+
+void SSequenceValidatorRules::Construct(const FArguments& InArgs)
+{
+	Validator = InArgs._Validator;
+
+	TSharedRef<FSequenceValidatorStyle> ValidatorStyle = FSequenceValidatorStyle::Get();
+
+	ChildSlot
+	[
+		SNew(SVerticalBox)
+		+SVerticalBox::Slot()
+		.AutoHeight()
+		[
+			SNew(STextBlock)
+			.Margin(4.f)
+			.Text(LOCTEXT("RulesTitle", "Rules"))
+			.TextStyle(ValidatorStyle, "SequenceValidator.PanelTitleText")
+		]
+		+SVerticalBox::Slot()
+		.FillHeight(1.f)
+		[
+			SAssignNew(ListView, SListView<FListItemPtr>)
+			.ListItemsSource(&ItemSource)
+			.OnGenerateRow(this, &SSequenceValidatorRules::OnListGenerateItemRow)
+		]
+	];
+
+	UpdateItemSource();
+}
+
+TSharedRef<ITableRow> SSequenceValidatorRules::OnListGenerateItemRow(FListItemPtr Item, const TSharedRef<STableViewBase>& OwnerTable)
+{
+	return SNew(SSequenceValidatorRuleEntry, OwnerTable)
+		.Item(Item);
+}
+
+void SSequenceValidatorRules::UpdateItemSource()
+{
+	ItemSource.Reset();
+	for (const FSequenceValidationRuleInfo& RuleInfo : Validator->GetRules())
+	{
+		TSharedRef<FSequenceValidationRuleInfo> NewRuleInfo = MakeShared<FSequenceValidationRuleInfo>(RuleInfo);
+		ItemSource.Add(NewRuleInfo);
+	}
+}
+
+}  // namespace UE::Sequencer
+
+#undef LOCTEXT_NAMESPACE
+
diff --git a/Engine/Plugins/MovieScene/SequenceValidator/Source/SequenceValidator/Private/Widgets/SSequenceValidatorRules.h b/Engine/Plugins/MovieScene/SequenceValidator/Source/SequenceValidator/Private/Widgets/SSequenceValidatorRules.h
new file mode 100644
--- /dev/null
+++ b/Engine/Plugins/MovieScene/SequenceValidator/Source/SequenceValidator/Private/Widgets/SSequenceValidatorRules.h
@@ -0,0 +1,42 @@
+// Copyright Epic Games, Inc. All Rights Reserved.
+
+#pragma once
+
+#include "Widgets/SCompoundWidget.h"
+#include "Widgets/Views/SListView.h"
+
+namespace UE::Sequencer
+{
+
+class FSequenceValidator;
+struct FSequenceValidationRuleInfo;
+
+class SSequenceValidatorRules : public SCompoundWidget
+{
+public:
+
+	SLATE_BEGIN_ARGS(SSequenceValidatorRules)
+	{}
+		SLATE_ARGUMENT(TSharedPtr<const FSequenceValidator>, Validator)
+	SLATE_END_ARGS()
+
+	void Construct(const FArguments& InArgs);
+
+private:
+
+	using FListItemPtr = TSharedPtr<FSequenceValidationRuleInfo>;
+
+	TSharedRef<ITableRow> OnListGenerateItemRow(FListItemPtr Item, const TSharedRef<STableViewBase>& OwnerTable);
+
+	void UpdateItemSource();
+
+private:
+
+	TSharedPtr<const FSequenceValidator> Validator;
+
+	TSharedPtr<SListView<FListItemPtr>> ListView;
+	TArray<FListItemPtr> ItemSource;
+};
+
+}  // namespace UE::Sequencer
+
diff --git a/Engine/Plugins/MovieScene/SequenceValidator/Source/SequenceValidator/Public/ISequenceValidatorModule.h b/Engine/Plugins/MovieScene/SequenceValidator/Source/SequenceValidator/Public/ISequenceValidatorModule.h
new file mode 100644
--- /dev/null
+++ b/Engine/Plugins/MovieScene/SequenceValidator/Source/SequenceValidator/Public/ISequenceValidatorModule.h
@@ -0,0 +1,35 @@
+// Copyright Epic Games, Inc. All Rights Reserved.
+
+#pragma once
+
+#include "Delegates/Delegate.h"
+#include "Logging/LogMacros.h"
+#include "Modules/ModuleInterface.h"
+
+DECLARE_LOG_CATEGORY_EXTERN(LogSequenceValidator, Log, All);
+
+namespace UE::Sequencer
+{
+
+class FSequenceValidationRule;
+struct FSequenceValidationRuleInfo;
+
+using FSequenceValidationRuleID = FDelegateHandle;
+
+/**
+ * Interface for the Sequence Validator module.
+ */
+class ISequenceValidatorModule : public IModuleInterface
+{
+public:
+
+	/** Registers a new validation rule factory. */
+	virtual FSequenceValidationRuleID RegisterValidationRule(FSequenceValidationRuleInfo&& InRuleInfo) = 0;
+	/** Gets the registered validation rule factories. */
+	virtual TConstArrayView<FSequenceValidationRuleInfo> GetValidationRules() const = 0;
+	/** Unregisters an existing validation rule factory. */
+	virtual void UnregisterValidationRule(FSequenceValidationRuleID InRuleID) = 0;
+};
+
+}  // namespace UE::Sequencer
+
diff --git a/Engine/Plugins/MovieScene/SequenceValidator/Source/SequenceValidator/Public/Validation/SequenceValidationResult.h b/Engine/Plugins/MovieScene/SequenceValidator/Source/SequenceValidator/Public/Validation/SequenceValidationResult.h
new file mode 100644
--- /dev/null
+++ b/Engine/Plugins/MovieScene/SequenceValidator/Source/SequenceValidator/Public/Validation/SequenceValidationResult.h
@@ -0,0 +1,122 @@
+// Copyright Epic Games, Inc. All Rights Reserved.
+
+#pragma once
+
+#include "Containers/Array.h"
+#include "Internationalization/Text.h"
+#include "Logging/TokenizedMessage.h"
+#include "Misc/FrameTime.h"
+#include "Misc/Optional.h"
+#include "Templates/SharedPointer.h"
+#include "UObject/WeakObjectPtr.h"
+
+class UMovieSceneSubSection;
+
+namespace UE::Sequencer
+{
+
+/**
+ * A single validation result, indicating a note, warning, error, etc. inside the sequence.
+ */
+class FSequenceValidationResult : public TSharedFromThis<FSequenceValidationResult>
+{
+public:
+
+	/** Creates a new default validation result. */
+	FSequenceValidationResult() = default;
+
+	/** Creates a new validation result for the given object. */
+	FSequenceValidationResult(UObject* InTarget)
+		: WeakTarget(InTarget)
+	{}
+
+	/** Creates a new validation result for the given object, with the given severity. */
+	FSequenceValidationResult(EMessageSeverity::Type InSeverity, UObject* InTarget)
+		: Severity(InSeverity)
+		, WeakTarget(InTarget)
+	{}
+
+	/** Creates a new validation result with the given severity and message. */
+	FSequenceValidationResult(EMessageSeverity::Type InSeverity, const FText& InUserMessage)
+		: Severity(InSeverity)
+		, UserMessage(InUserMessage)
+	{}
+
+	/** Creates a new validation result with the given severity and message, for the given object. */
+	FSequenceValidationResult(EMessageSeverity::Type InSeverity, UObject* InTarget, const FText& InUserMessage)
+		: Severity(InSeverity)
+		, WeakTarget(InTarget)
+		, UserMessage(InUserMessage)
+	{}
+
+public:
+
+	EMessageSeverity::Type GetSeverity() const { return Severity; }
+
+	SEQUENCEVALIDATOR_API bool GetSubSectionTrail(TArray<UMovieSceneSubSection*>& OutTrail) const;
+	SEQUENCEVALIDATOR_API void SetSubSectionTrail(TArrayView<UMovieSceneSubSection*> InTrail);
+
+	UObject* GetTarget() const { return WeakTarget.Get(); }
+	void SetTarget(UObject* InTarget) { WeakTarget = InTarget; }
+
+	const FText& GetUserMessage() const { return UserMessage; }
+	void SetUserMessage(const FText& InUserMessage) { UserMessage = InUserMessage; }
+
+	bool HasLocalTime() const { return LocalTime.IsSet(); }
+	FFrameTime GetLocalTime() const { return LocalTime.Get(FFrameTime(0)); }
+	void SetLocalTime(const FFrameTime InLocalTime) { LocalTime = InLocalTime; }
+
+public:
+
+	TSharedPtr<FSequenceValidationResult> GetParent() const
+	{
+		return WeakParent.Pin();
+	}
+
+	TSharedPtr<FSequenceValidationResult> GetRoot();
+
+	bool HasChildren() const { return !Children.IsEmpty(); }
+
+	TArrayView<const TSharedPtr<FSequenceValidationResult>> GetChildren() const
+	{
+		return Children;
+	}
+
+	SEQUENCEVALIDATOR_API void AddChild(TSharedRef<FSequenceValidationResult> InChild);
+
+	SEQUENCEVALIDATOR_API void AppendChildren(TConstArrayView<TSharedPtr<FSequenceValidationResult>> InChildren);
+
+private:
+
+	TWeakPtr<FSequenceValidationResult> WeakParent;
+	TArray<TSharedPtr<FSequenceValidationResult>> Children;
+
+	EMessageSeverity::Type Severity = EMessageSeverity::Info;
+	TArray<TWeakObjectPtr<UMovieSceneSubSection>> WeakSubSectionTrail;
+	TWeakObjectPtr<> WeakTarget;
+	FText UserMessage;
+	TOptional<FFrameTime> LocalTime;
+};
+
+/**
+ * A collection of validation results.
+ */
+class FSequenceValidationResults : public TSharedFromThis<FSequenceValidationResult>
+{
+public:
+
+	SEQUENCEVALIDATOR_API void AddResult(TSharedRef<FSequenceValidationResult> InResult);
+	SEQUENCEVALIDATOR_API void AppendResults(const FSequenceValidationResults& InResults);
+	SEQUENCEVALIDATOR_API void AppendResults(TConstArrayView<TSharedPtr<FSequenceValidationResult>> InResults);
+
+	SEQUENCEVALIDATOR_API TArrayView<const TSharedPtr<FSequenceValidationResult>> GetResults() const;
+
+	SEQUENCEVALIDATOR_API void Reset();
+
+private:
+
+	TArray<TSharedPtr<FSequenceValidationResult>> ValidationResults;
+};
+
+}  // namespace UE::Sequencer
+
diff --git a/Engine/Plugins/MovieScene/SequenceValidator/Source/SequenceValidator/Public/Validation/SequenceValidationRule.h b/Engine/Plugins/MovieScene/SequenceValidator/Source/SequenceValidator/Public/Validation/SequenceValidationRule.h
new file mode 100644
--- /dev/null
+++ b/Engine/Plugins/MovieScene/SequenceValidator/Source/SequenceValidator/Public/Validation/SequenceValidationRule.h
@@ -0,0 +1,51 @@
+// Copyright Epic Games, Inc. All Rights Reserved.
+
+#pragma once
+
+#include "Internationalization/Text.h"
+#include "Templates/SharedPointer.h"
+
+class UMovieSceneSequence;
+
+namespace UE::Sequencer
+{
+
+class FSequenceValidationResults;
+class FSequenceValidationRule;
+
+DECLARE_DELEGATE_RetVal(TSharedRef<FSequenceValidationRule>, FOnCreateSequenceValidationRule);
+
+/** Basic information regarding a validation rule, and how to create it. */
+struct FSequenceValidationRuleInfo
+{
+	/** The name of the rule. */
+	FText RuleName;
+	/** The description of the rule. */
+	FText RuleDescription;
+	/** A factory method for creating an instance of the rule. */
+	FOnCreateSequenceValidationRule RuleFactory;
+	/** Whether this rule is enabled. */
+	bool bIsEnabled = true;
+};
+
+/**
+ * Base class for a validation rule that can be run on a sequence.
+ */
+class FSequenceValidationRule : public TSharedFromThis<FSequenceValidationRule>
+{
+public:
+
+	FSequenceValidationRule() = default;
+	virtual ~FSequenceValidationRule() = default;
+
+public:
+
+	void Run(const UMovieSceneSequence* InSequence, FSequenceValidationResults& OutResults) const;
+
+protected:
+
+	virtual void OnRun(const UMovieSceneSequence* InSequence, FSequenceValidationResults& OutResults) const {}
+};
+
+}  // namespace UE::Sequencer
+
diff --git a/Engine/Plugins/MovieScene/SequenceValidator/Source/SequenceValidator/Public/Validation/SequenceValidator.h b/Engine/Plugins/MovieScene/SequenceValidator/Source/SequenceValidator/Public/Validation/SequenceValidator.h
new file mode 100644
--- /dev/null
+++ b/Engine/Plugins/MovieScene/SequenceValidator/Source/SequenceValidator/Public/Validation/SequenceValidator.h
@@ -0,0 +1,89 @@
+// Copyright Epic Games, Inc. All Rights Reserved.
+
+#pragma once
+
+#include "Containers/Array.h"
+#include "Containers/ArrayView.h"
+#include "Delegates/Delegate.h"
+#include "ISequenceValidatorModule.h"
+#include "Templates/SharedPointerFwd.h"
+#include "Templates/UniquePtr.h"
+#include "Validation/SequenceValidationResult.h"
+
+class UMovieSceneSequence;
+
+namespace UE::Sequencer
+{
+
+class FSequenceValidationRule;
+
+namespace Internal { class FSequenceValidatorTaskScheduler; }
+
+/**
+ * A sequence validator that can run a given set of validation rules on a sequence.
+ *
+ * Validation rules are created along with the validator instance, using whatever validation
+ * rules factories are registered on the ISequenceValidatorModule at the time.
+ */
+class FSequenceValidator
+{
+public:
+
+	FSequenceValidator();
+	~FSequenceValidator();
+
+public:
+
+	/** The current queue of sequences to validate. */
+	TArrayView<UMovieSceneSequence* const> GetQueue() const;
+	/** The list of validation rules to run on the queue. */
+	TConstArrayView<FSequenceValidationRuleInfo> GetRules() const;
+	/** The current list of results generated by the validation rules. */
+	const FSequenceValidationResults& GetResults() const;
+
+	/** Add the given root sequence to the queue. */
+	void Queue(UMovieSceneSequence* InSequence);
+	/** Add the given root sequences to the queue. */
+	void Queue(TArrayView<UMovieSceneSequence*> InSequences);
+	/** Clears the current queue. */
+	void ClearQueue();
+
+	/**
+	 * Start validating the current queue.
+	 * This runs all validation rules on all queued-up sequences in an asynchronous manner.
+	 * Call the IsValidating() method to know when the validation has ended.
+	 */
+	void StartValidation();
+
+	/** Whether an asynchronous validation is in progress. */
+	bool IsValidating() const;
+
+	/** Gets an event that will be triggered when an asynchronous validation ends. */
+	FSimpleMulticastDelegate& GetOnValidationFinished();
+
+public:
+
+	/** Validates the given sequence synchronously. */
+	void Validate(UMovieSceneSequence* InSequence);
+	/** Validates the given sequences synchronously. */
+	void Validate(TArrayView<UMovieSceneSequence*> InSequences);
+
+private:
+
+	void CreateValidationRuleInfos();
+	void StartValidation(bool bAsync);
+	void OnValidationFinished();
+
+private:
+
+	TArray<UMovieSceneSequence*> ValidationQueue;
+	TArray<FSequenceValidationRuleInfo> ValidationRuleInfos;
+	FSequenceValidationResults ValidationResults;
+
+	TUniquePtr<Internal::FSequenceValidatorTaskScheduler> Scheduler;
+
+	FSimpleMulticastDelegate OnValidationFinishedDelegate;
+};
+
+}  // namespace UE::Sequencer
+
diff --git a/Engine/Plugins/MovieScene/SequenceValidator/Source/SequenceValidator/SequenceValidator.Build.cs b/Engine/Plugins/MovieScene/SequenceValidator/Source/SequenceValidator/SequenceValidator.Build.cs
new file mode 100644
--- /dev/null
+++ b/Engine/Plugins/MovieScene/SequenceValidator/Source/SequenceValidator/SequenceValidator.Build.cs
@@ -0,0 +1,34 @@
+// Copyright Epic Games, Inc. All Rights Reserved.
+
+using UnrealBuildTool;
+
+public class SequenceValidator : ModuleRules
+{
+	public SequenceValidator(ReadOnlyTargetRules Target) : base(Target)
+	{
+		PrivateDependencyModuleNames.AddRange(
+			new string[]
+			{
+				"ApplicationCore",
+				"AssetRegistry",
+				"ContentBrowser",
+				"Core",
+				"CoreUObject",
+				"Engine",
+				"InputCore",
+				"LevelSequence",
+				"LevelSequenceEditor",
+				"MovieScene",
+				"MovieSceneTracks",
+				"Projects",
+				"Sequencer",
+				"Slate",
+				"SlateCore",
+				"ToolMenus",
+				"UnrealEd",
+				"WorkspaceMenuStructure",
+			}
+		);
+	}
+}
+
