Sequencer: add some checks for the static analyzer

diff --git a/Engine/Plugins/MovieScene/SequenceValidator/Source/SequenceValidator/Private/Validation/SequenceValidator.cpp b/Engine/Plugins/MovieScene/SequenceValidator/Source/SequenceValidator/Private/Validation/SequenceValidator.cpp
--- a/Engine/Plugins/MovieScene/SequenceValidator/Source/SequenceValidator/Private/Validation/SequenceValidator.cpp
+++ b/Engine/Plugins/MovieScene/SequenceValidator/Source/SequenceValidator/Private/Validation/SequenceValidator.cpp
@@ -186,7 +186,10 @@ void FSequenceValidatorTaskScheduler::RunQueue(TArrayView<UMovieSceneSequence*>
 		Runner.Start(*this);
 	}
 
-	MainSignal->Wait();
+	if (ensure(MainSignal))
+	{
+		MainSignal->Wait();
+	}
 
 	UE_LOG(LogSequenceValidator, Log, TEXT("All sequence validator runners finished!"));
 	const bool bAllRunnersFinished = Algo::AllOf(Runners, 
@@ -196,9 +199,11 @@ void FSequenceValidatorTaskScheduler::RunQueue(TArrayView<UMovieSceneSequence*>
 			});
 	ensure(bAllRunnersFinished);
 
-	ensure(MainSignal);
-	FPlatformProcess::ReturnSynchEventToPool(MainSignal);
-	MainSignal = nullptr;
+	if (ensure(MainSignal))
+	{
+		FPlatformProcess::ReturnSynchEventToPool(MainSignal);
+		MainSignal = nullptr;
+	}
 
 	GatherResults(InQueue);
 
