Sequencer: Add validator for Missing Assets

diff --git a/Engine/Plugins/MovieScene/SequenceValidator/Source/SequenceValidator/Private/SequenceValidatorModule.cpp b/Engine/Plugins/MovieScene/SequenceValidator/Source/SequenceValidator/Private/SequenceValidatorModule.cpp
--- a/Engine/Plugins/MovieScene/SequenceValidator/Source/SequenceValidator/Private/SequenceValidatorModule.cpp
+++ b/Engine/Plugins/MovieScene/SequenceValidator/Source/SequenceValidator/Private/SequenceValidatorModule.cpp
@@ -4,6 +4,7 @@
 
 #include "Editor.h"
 #include "SequenceValidatorCommands.h"
+#include "Validation/Rules/SequenceValidationRule_MissingAssets.h"
 #include "Validation/Rules/SequenceValidationRule_SectionAlignments.h"
 #include "Validation/Rules/SequenceValidationRule_WholeSectionRanges.h"
 #include "Widgets/SSequenceValidator.h"
@@ -75,10 +76,12 @@ void FSequenceValidatorModule::OnPostEngineInit()
 
 void FSequenceValidatorModule::RegisterBuiltInRules()
 {
+	BuiltInValidationRules.Add(
+		RegisterValidationRule(FSequenceValidationRule_MissingAssets::MakeRuleInfo()));
 	BuiltInValidationRules.Add(
 			RegisterValidationRule(FSequenceValidationRule_SectionAlignments::MakeRuleInfo()));
 	BuiltInValidationRules.Add(
-			RegisterValidationRule(FSequenceValidationRule_WholeSectionRanges::MakeRuleInfo()));
+		RegisterValidationRule(FSequenceValidationRule_WholeSectionRanges::MakeRuleInfo()));
 }
 
 void FSequenceValidatorModule::UnregisterBuiltInRules()
diff --git a/Engine/Plugins/MovieScene/SequenceValidator/Source/SequenceValidator/Private/Validation/Rules/SequenceValidationRule_MissingAssets.cpp b/Engine/Plugins/MovieScene/SequenceValidator/Source/SequenceValidator/Private/Validation/Rules/SequenceValidationRule_MissingAssets.cpp
new file mode 100644
--- /dev/null
+++ b/Engine/Plugins/MovieScene/SequenceValidator/Source/SequenceValidator/Private/Validation/Rules/SequenceValidationRule_MissingAssets.cpp
@@ -0,0 +1,137 @@
+// Copyright Epic Games, Inc. All Rights Reserved.
+
+#include "Validation/Rules/SequenceValidationRule_MissingAssets.h"
+
+#include "Algo/BinarySearch.h"
+#include "Internationalization/Internationalization.h"
+#include "MovieScene.h"
+#include "MovieSceneSection.h"
+#include "MovieSceneSequence.h"
+#include "MovieSceneSequenceVisitor.h"
+#include "MovieSceneTrack.h"
+#include "Sections/MovieSceneAudioSection.h"
+#include "Sections/MovieScene3DConstraintSection.h"
+#include "Sections/MovieSceneCameraCutSection.h"
+#include "Sections/MovieSceneSkeletalAnimationSection.h"
+#include "Sections/MovieSceneSubSection.h"
+#include "Validation/SequenceValidationResult.h"
+
+#define LOCTEXT_NAMESPACE "SequenceValidationRule_MissingAssets"
+
+namespace UE::Sequencer
+{
+
+struct FMissingAssetsValidationVisitor : public UE::MovieScene::ISequenceVisitor
+{
+	using FSubSequenceSpace = UE::MovieScene::FSubSequenceSpace;
+
+public:
+
+	FMissingAssetsValidationVisitor(FSequenceValidationResults& OutResults)
+		: Results(OutResults)
+	{
+	}
+
+	void Initialize(const UMovieSceneSequence* InSequence)
+	{
+		UMovieScene* MovieScene = InSequence->GetMovieScene();
+
+		CurrentSubSectionTrail.Reserve(4);
+	}
+
+	virtual void VisitStartSubSequence(UMovieSceneSubSection* SubSection, const FGuid& ObjectBindingID, const FSubSequenceSpace& LocalSpace) override
+	{
+		CurrentSubSectionTrail.Add(SubSection);
+	}
+
+	virtual void VisitEndSubSequence(UMovieSceneSubSection* SubSection, const FGuid& ObjectBindingID, const FSubSequenceSpace& LocalSpace) override
+	{
+		ensure(CurrentSubSectionTrail.Num() > 0 && CurrentSubSectionTrail.Last() == SubSection);
+		CurrentSubSectionTrail.Pop(EAllowShrinking::No);
+	}
+
+	virtual void VisitSection(UMovieSceneTrack* Track, UMovieSceneSection* Section, const FGuid& ObjectBindingID, const FSubSequenceSpace& LocalSpace) override
+	{
+		if (UMovieSceneAudioSection* AudioSection = Cast<UMovieSceneAudioSection>(Section))
+		{
+			if (!AudioSection->GetSound())
+			{
+				WarnAbout(Results, Section);
+			}
+		}
+		else if (UMovieSceneCameraCutSection* CameraCutSection = Cast<UMovieSceneCameraCutSection>(Section))
+		{
+			if (!CameraCutSection->GetCameraBindingID().IsValid())
+			{
+				WarnAbout(Results, Section);
+			}
+		}
+		else if (UMovieScene3DConstraintSection* ConstraintSection = Cast<UMovieScene3DConstraintSection>(Section))
+		{
+			if (!ConstraintSection->GetConstraintBindingID().IsValid())
+			{
+				WarnAbout(Results, Section);
+			}
+		}
+		else if (UMovieSceneSubSection* SubSection = Cast<UMovieSceneSubSection>(Section))
+		{
+			if (!SubSection->GetSequence())
+			{
+				WarnAbout(Results, Section);
+			}
+		}
+		else if (UMovieSceneSkeletalAnimationSection* SkeletalAnimationSection = Cast<UMovieSceneSkeletalAnimationSection>(Section))
+		{
+			if (!SkeletalAnimationSection->GetAnimation())
+			{
+				WarnAbout(Results, Section);
+			}
+		}
+	}
+
+private:
+
+	void WarnAbout(FSequenceValidationResults& OutResults, UObject* InTarget = nullptr)
+	{
+		TSharedRef<FSequenceValidationResult> NewResult = MakeShared<FSequenceValidationResult>(EMessageSeverity::Warning, InTarget);
+		NewResult->SetSubSectionTrail(CurrentSubSectionTrail);
+		NewResult->SetUserMessage(FText::Format(
+			LOCTEXT("MissingAssetWarningFmt", "Section {0} has a missing asset."),
+			FText::FromString(InTarget->GetName())));
+		OutResults.AddResult(NewResult);
+	}
+
+private:
+
+	TArray<UMovieSceneSubSection*> CurrentSubSectionTrail;
+	FSequenceValidationResults& Results;
+};
+
+FSequenceValidationRuleInfo FSequenceValidationRule_MissingAssets::MakeRuleInfo()
+{
+	FSequenceValidationRuleInfo RuleInfo;
+	RuleInfo.RuleName = LOCTEXT("RuleName", "Missing Assets");
+	RuleInfo.RuleDescription = LOCTEXT("RuleDescription", "Check that animations, audio, subsequences, camera cuts, attach/path sections reference valid assets.");
+	RuleInfo.RuleFactory = FOnCreateSequenceValidationRule::CreateLambda([]() { return MakeShared<FSequenceValidationRule_MissingAssets>(); });
+	return RuleInfo;
+}
+
+void FSequenceValidationRule_MissingAssets::OnRun(const UMovieSceneSequence* InSequence, FSequenceValidationResults& OutResults) const
+{
+	UMovieScene* MovieScene = InSequence->GetMovieScene();
+
+	// Visit the sequence hierarchy and report on things that are off.
+	UE::MovieScene::FSequenceVisitParams VisitParams;
+	VisitParams.bVisitRootTracks = true;
+	VisitParams.bVisitSubSequences = true;
+	VisitParams.bVisitSections = true;
+
+	FMissingAssetsValidationVisitor Validator(OutResults);
+	Validator.Initialize(InSequence);
+	UE::MovieScene::VisitSequence(const_cast<UMovieSceneSequence*>(InSequence), VisitParams, Validator);
+}
+
+}  // namespace UE::Sequencer
+
+#undef LOCTEXT_NAMESPACE
+
diff --git a/Engine/Plugins/MovieScene/SequenceValidator/Source/SequenceValidator/Private/Validation/Rules/SequenceValidationRule_MissingAssets.h b/Engine/Plugins/MovieScene/SequenceValidator/Source/SequenceValidator/Private/Validation/Rules/SequenceValidationRule_MissingAssets.h
new file mode 100644
--- /dev/null
+++ b/Engine/Plugins/MovieScene/SequenceValidator/Source/SequenceValidator/Private/Validation/Rules/SequenceValidationRule_MissingAssets.h
@@ -0,0 +1,25 @@
+// Copyright Epic Games, Inc. All Rights Reserved.
+
+#pragma once
+
+#include "Validation/SequenceValidationRule.h"
+
+namespace UE::Sequencer
+{
+
+/**
+ * A validation rule that checks that sections reference valid assets.
+ */
+class FSequenceValidationRule_MissingAssets : public FSequenceValidationRule
+{
+public:
+
+	static FSequenceValidationRuleInfo MakeRuleInfo();
+
+protected:
+
+	virtual void OnRun(const UMovieSceneSequence* InSequence, FSequenceValidationResults& OutResults) const override;
+};
+
+}  // namespace UE::Sequencer
+
