Sequencer: Validation UI Tweaks

- Added RuleInfo to ValidationResult so that the result can look up which rule it came from
- Added RuleColor from a map of valid colors for each RuleInfo
- Rearrange Queue, Rules, Results, add icons to titles
- Move Add button to within Queue
- Move Start Validation to bottom right
- Show combined RuleInfo as the parent result
- Added Remove and Clear to queue right click menu
- Green tick if Result has no issues

#rnx

diff --git a/Engine/Plugins/MovieScene/SequenceValidator/Source/SequenceValidator/Private/SequenceValidatorStyle.cpp b/Engine/Plugins/MovieScene/SequenceValidator/Source/SequenceValidator/Private/SequenceValidatorStyle.cpp
--- a/Engine/Plugins/MovieScene/SequenceValidator/Source/SequenceValidator/Private/SequenceValidatorStyle.cpp
+++ b/Engine/Plugins/MovieScene/SequenceValidator/Source/SequenceValidator/Private/SequenceValidatorStyle.cpp
@@ -29,15 +29,15 @@ FSequenceValidatorStyle::FSequenceValidatorStyle()
 	const FTableRowStyle AlternatingTableRowStyle = CoreStyle.GetWidgetStyle<FTableRowStyle>("TableView.AlternatingRow");
 
 	// Validation rules icons.
-	Set("ValidationRule", new IMAGE_BRUSH_SVG("Icons/ValidationRule", Icon16x16));
-
 	Set("ValidationRule.RowStyle", AlternatingTableRowStyle);
+	Set("ValidationRule.Rule", new IMAGE_BRUSH("Icons/ValidationRule-Rule", Icon16x16));
 
 	// Validation results icons.
 	Set("ValidationResult.Key", new IMAGE_BRUSH("Icons/ValidationResult-Key", Icon16x16));
 	Set("ValidationResult.Section", new IMAGE_BRUSH("Icons/ValidationResult-Section", Icon16x16));
 	Set("ValidationResult.Track", new IMAGE_BRUSH("Icons/ValidationResult-Track", Icon16x16));
 	Set("ValidationResult.Sequence", new IMAGE_BRUSH("Icons/ValidationResult-Sequence", Icon16x16));
+	Set("ValidationResult.Pass", new IMAGE_BRUSH_SVG("Icons/ValidationResult-Pass", Icon16x16));
 
 	Set("ValidationResult.RowStyle", AlternatingTableRowStyle);
 
@@ -48,11 +48,16 @@ FSequenceValidatorStyle::FSequenceValidatorStyle()
 	// Tab icons.
 	Set("SequenceValidator.TabIcon", new IMAGE_BRUSH_SVG("Icons/SequenceValidation", Icon16x16));
 
+	// Title icons.
+	Set("SequenceValidator.QueueTitleIcon", new IMAGE_BRUSH_SVG("Icons/ValidationQueue", Icon16x16));
+	Set("SequenceValidator.RulesTitleIcon", new IMAGE_BRUSH_SVG("Icons/ValidationRule", Icon16x16));
+	Set("SequenceValidator.ResultsTitleIcon", new IMAGE_BRUSH_SVG("Icons/ValidationResults", Icon16x16));
+
 	// Text styles.
 
 	Set("SequenceValidator.PanelTitleText", FTextBlockStyle(NormalText)
-		.SetFont(DEFAULT_FONT("Bold", 18))
-		.SetColorAndOpacity(FLinearColor(1.0f, 1.0f, 1.0f))
+		.SetFont(DEFAULT_FONT("Regular", 16))
+		.SetColorAndOpacity(FLinearColor(0.8f, 0.8f, 0.8f))
 		.SetHighlightColor(FLinearColor(1.0f, 1.0f, 1.0f))
 		.SetShadowOffset(FVector2f::UnitVector)
 		.SetShadowColorAndOpacity(FLinearColor(0, 0, 0, 0.9f)));
diff --git a/Engine/Plugins/MovieScene/SequenceValidator/Source/SequenceValidator/Private/Validation/Rules/SequenceValidationRule_MissingAssets.cpp b/Engine/Plugins/MovieScene/SequenceValidator/Source/SequenceValidator/Private/Validation/Rules/SequenceValidationRule_MissingAssets.cpp
--- a/Engine/Plugins/MovieScene/SequenceValidator/Source/SequenceValidator/Private/Validation/Rules/SequenceValidationRule_MissingAssets.cpp
+++ b/Engine/Plugins/MovieScene/SequenceValidator/Source/SequenceValidator/Private/Validation/Rules/SequenceValidationRule_MissingAssets.cpp
@@ -27,8 +27,9 @@ struct FMissingAssetsValidationVisitor : public UE::MovieScene::ISequenceVisitor
 
 public:
 
-	FMissingAssetsValidationVisitor(FSequenceValidationResults& OutResults)
-		: Results(OutResults)
+	FMissingAssetsValidationVisitor(const FSequenceValidationRuleInfo& InRuleInfo, FSequenceValidationResults& OutResults)
+		: RuleInfo(InRuleInfo)
+		, Results(OutResults)
 	{
 	}
 
@@ -93,7 +94,7 @@ private:
 
 	void WarnAbout(FSequenceValidationResults& OutResults, UObject* InTarget = nullptr)
 	{
-		TSharedRef<FSequenceValidationResult> NewResult = MakeShared<FSequenceValidationResult>(EMessageSeverity::Warning, InTarget);
+		TSharedRef<FSequenceValidationResult> NewResult = MakeShared<FSequenceValidationResult>(EMessageSeverity::Warning, InTarget, RuleInfo);
 		NewResult->SetSubSectionTrail(CurrentSubSectionTrail);
 		NewResult->SetUserMessage(FText::Format(
 			LOCTEXT("MissingAssetWarningFmt", "Section {0} has a missing asset."),
@@ -103,6 +104,7 @@ private:
 
 private:
 
+	FSequenceValidationRuleInfo RuleInfo;
 	TArray<UMovieSceneSubSection*> CurrentSubSectionTrail;
 	FSequenceValidationResults& Results;
 };
@@ -113,6 +115,7 @@ FSequenceValidationRuleInfo FSequenceValidationRule_MissingAssets::MakeRuleInfo(
 	RuleInfo.RuleName = LOCTEXT("RuleName", "Missing Assets");
 	RuleInfo.RuleDescription = LOCTEXT("RuleDescription", "Check that animations, audio, subsequences, camera cuts, attach/path sections reference valid assets.");
 	RuleInfo.RuleFactory = FOnCreateSequenceValidationRule::CreateLambda([]() { return MakeShared<FSequenceValidationRule_MissingAssets>(); });
+	RuleInfo.RuleColor = FSequenceValidationRule::GetRuleColor(RuleInfo.RuleName);
 	return RuleInfo;
 }
 
@@ -126,7 +129,7 @@ void FSequenceValidationRule_MissingAssets::OnRun(const UMovieSceneSequence* InS
 	VisitParams.bVisitSubSequences = true;
 	VisitParams.bVisitSections = true;
 
-	FMissingAssetsValidationVisitor Validator(OutResults);
+	FMissingAssetsValidationVisitor Validator(MakeRuleInfo(), OutResults);
 	Validator.Initialize(InSequence);
 	UE::MovieScene::VisitSequence(const_cast<UMovieSceneSequence*>(InSequence), VisitParams, Validator);
 }
diff --git a/Engine/Plugins/MovieScene/SequenceValidator/Source/SequenceValidator/Private/Validation/Rules/SequenceValidationRule_SectionAlignments.cpp b/Engine/Plugins/MovieScene/SequenceValidator/Source/SequenceValidator/Private/Validation/Rules/SequenceValidationRule_SectionAlignments.cpp
--- a/Engine/Plugins/MovieScene/SequenceValidator/Source/SequenceValidator/Private/Validation/Rules/SequenceValidationRule_SectionAlignments.cpp
+++ b/Engine/Plugins/MovieScene/SequenceValidator/Source/SequenceValidator/Private/Validation/Rules/SequenceValidationRule_SectionAlignments.cpp
@@ -104,8 +104,9 @@ struct FSectionAlignmentsValidationVisitor : public UE::MovieScene::ISequenceVis
 
 public:
 
-	FSectionAlignmentsValidationVisitor(TArrayView<FImportantFrameTime> InImportantTimes, FSequenceValidationResults& OutResults)
-		: ImportantTimes(InImportantTimes)
+	FSectionAlignmentsValidationVisitor(const FSequenceValidationRuleInfo& InRuleInfo, TArrayView<FImportantFrameTime> InImportantTimes, FSequenceValidationResults& OutResults)
+		: RuleInfo(InRuleInfo)
+		, ImportantTimes(InImportantTimes)
 		, Results(OutResults)
 	{
 	}
@@ -188,7 +189,7 @@ private:
 
 		if (bDoWarn)
 		{
-			TSharedRef<FSequenceValidationResult> NewResult = MakeShared<FSequenceValidationResult>(EMessageSeverity::Warning, InTarget);
+			TSharedRef<FSequenceValidationResult> NewResult = MakeShared<FSequenceValidationResult>(EMessageSeverity::Warning, InTarget, RuleInfo);
 			NewResult->SetSubSectionTrail(CurrentSubSectionTrail);
 			NewResult->SetLocalTime(InTime);
 			NewResult->SetUserMessage(FText::Format(
@@ -200,6 +201,7 @@ private:
 
 private:
 
+	FSequenceValidationRuleInfo RuleInfo;
 	FFrameTime Threshold;
 	TArray<FImportantFrameTime> ImportantTimes;
 	TArray<UMovieSceneSubSection*> CurrentSubSectionTrail;
@@ -212,6 +214,7 @@ FSequenceValidationRuleInfo FSequenceValidationRule_SectionAlignments::MakeRuleI
 	RuleInfo.RuleName = LOCTEXT("RuleName", "Section Alignments");
 	RuleInfo.RuleDescription = LOCTEXT("RuleDescription", "Check that sections are not off by a couple frames from notable frames.");
 	RuleInfo.RuleFactory = FOnCreateSequenceValidationRule::CreateLambda([]() { return MakeShared<FSequenceValidationRule_SectionAlignments>(); });
+	RuleInfo.RuleColor = FSequenceValidationRule::GetRuleColor(RuleInfo.RuleName);
 	return RuleInfo;
 }
 
@@ -241,7 +244,7 @@ void FSequenceValidationRule_SectionAlignments::OnRun(const UMovieSceneSequence*
 	VisitParams.bVisitSections = true;
 	VisitParams.bVisitTracks = true;
 
-	FSectionAlignmentsValidationVisitor Validator(ImportantTimes, OutResults);
+	FSectionAlignmentsValidationVisitor Validator(MakeRuleInfo(), ImportantTimes, OutResults);
 	Validator.Initialize(InSequence);
 	UE::MovieScene::VisitSequence(const_cast<UMovieSceneSequence*>(InSequence), VisitParams, Validator);
 }
diff --git a/Engine/Plugins/MovieScene/SequenceValidator/Source/SequenceValidator/Private/Validation/Rules/SequenceValidationRule_WholeSectionRanges.cpp b/Engine/Plugins/MovieScene/SequenceValidator/Source/SequenceValidator/Private/Validation/Rules/SequenceValidationRule_WholeSectionRanges.cpp
--- a/Engine/Plugins/MovieScene/SequenceValidator/Source/SequenceValidator/Private/Validation/Rules/SequenceValidationRule_WholeSectionRanges.cpp
+++ b/Engine/Plugins/MovieScene/SequenceValidator/Source/SequenceValidator/Private/Validation/Rules/SequenceValidationRule_WholeSectionRanges.cpp
@@ -20,8 +20,9 @@ class FWholeSectionRangeVisitor : public UE::MovieScene::ISequenceVisitor
 
 public:
 
-	FWholeSectionRangeVisitor(FSequenceValidationResults& OutResults)
-		: Results(OutResults)
+	FWholeSectionRangeVisitor(const FSequenceValidationRuleInfo& InRuleInfo, FSequenceValidationResults& OutResults)
+		: RuleInfo(InRuleInfo)
+		, Results(OutResults)
 	{
 	}
 
@@ -65,7 +66,7 @@ private:
 
 	void AddWarning(UMovieSceneSection* Section, FFrameTime LocalTime)
 	{
-		TSharedRef<FSequenceValidationResult> NewResult = MakeShared<FSequenceValidationResult>(EMessageSeverity::Warning, Section);
+		TSharedRef<FSequenceValidationResult> NewResult = MakeShared<FSequenceValidationResult>(EMessageSeverity::Warning, Section, RuleInfo);
 		NewResult->SetSubSectionTrail(CurrentSubSectionTrail);
 		NewResult->SetLocalTime(LocalTime);
 		NewResult->SetUserMessage(LOCTEXT("WholeFrameWarningFmt", "Section bound is not on a whole frame."));
@@ -74,6 +75,7 @@ private:
 
 private:
 
+	FSequenceValidationRuleInfo RuleInfo;
 	TArray<UMovieSceneSubSection*> CurrentSubSectionTrail;
 	FSequenceValidationResults& Results;
 };
@@ -84,6 +86,7 @@ FSequenceValidationRuleInfo FSequenceValidationRule_WholeSectionRanges::MakeRule
 	RuleInfo.RuleName = LOCTEXT("RuleName", "Whole Section Ranges");
 	RuleInfo.RuleDescription = LOCTEXT("RuleDescription", "Check that sections start and end on whole frames.");
 	RuleInfo.RuleFactory = FOnCreateSequenceValidationRule::CreateLambda([]() { return MakeShared<FSequenceValidationRule_WholeSectionRanges>(); });
+	RuleInfo.RuleColor = FSequenceValidationRule::GetRuleColor(RuleInfo.RuleName);
 	return RuleInfo;
 }
 
@@ -96,7 +99,7 @@ void FSequenceValidationRule_WholeSectionRanges::OnRun(const UMovieSceneSequence
 	VisitParams.bVisitTracks = true;
 	VisitParams.bVisitSections = true;
 
-	FWholeSectionRangeVisitor Visitor(OutResults);
+	FWholeSectionRangeVisitor Visitor(MakeRuleInfo(), OutResults);
 	UE::MovieScene::VisitSequence(const_cast<UMovieSceneSequence*>(InSequence), VisitParams, Visitor);
 }
 
diff --git a/Engine/Plugins/MovieScene/SequenceValidator/Source/SequenceValidator/Private/Validation/SequenceValidationRule.cpp b/Engine/Plugins/MovieScene/SequenceValidator/Source/SequenceValidator/Private/Validation/SequenceValidationRule.cpp
--- a/Engine/Plugins/MovieScene/SequenceValidator/Source/SequenceValidator/Private/Validation/SequenceValidationRule.cpp
+++ b/Engine/Plugins/MovieScene/SequenceValidator/Source/SequenceValidator/Private/Validation/SequenceValidationRule.cpp
@@ -5,6 +5,44 @@
 namespace UE::Sequencer
 {
 
+FSlateColor FSequenceValidationRule::GetRuleColor(const FText& InRuleName)
+{
+	if (InRuleName.IsEmpty())
+	{
+		return FSlateColor::UseForeground();
+	}
+
+	static TMap<FString, FSlateColor> RuleColorMap;
+	if (!RuleColorMap.Contains(InRuleName.ToString()))
+	{
+		static int32 Count = 0;
+		static TArray<FSlateColor> AvailableRuleColors =
+		{
+			FSlateColor(FLinearColor::FromSRGBColor(FColor(253, 94, 94, 255))),
+			FSlateColor(FLinearColor::FromSRGBColor(FColor(225, 102, 182, 255))),
+			FSlateColor(FLinearColor::FromSRGBColor(FColor(187, 107, 240, 255))),
+			FSlateColor(FLinearColor::FromSRGBColor(FColor(134, 138, 253, 255))),
+			FSlateColor(FLinearColor::FromSRGBColor(FColor(51, 191, 255, 255))),
+			FSlateColor(FLinearColor::FromSRGBColor(FColor(68, 213, 191, 255))),
+			FSlateColor(FLinearColor::FromSRGBColor(FColor(95, 227, 103, 255))),
+			FSlateColor(FLinearColor::FromSRGBColor(FColor(170, 231, 85, 255))),
+			FSlateColor(FLinearColor::FromSRGBColor(FColor(247, 216, 43, 255))),
+			FSlateColor(FLinearColor::FromSRGBColor(FColor(249, 169, 31, 255))),
+			FSlateColor(FLinearColor::FromSRGBColor(FColor(254, 133, 57, 255))),
+			FSlateColor(FLinearColor::FromSRGBColor(FColor(169, 129, 102, 255)))
+		};
+			
+		RuleColorMap.Add(InRuleName.ToString(), AvailableRuleColors[Count]);
+
+		Count++;
+		if (Count >= AvailableRuleColors.Num())
+		{
+			Count = 0;
+		}
+	}
+	return RuleColorMap[InRuleName.ToString()];
+}
+
 void FSequenceValidationRule::Run(const UMovieSceneSequence* InSequence, FSequenceValidationResults& OutResults) const
 {
 	OnRun(InSequence, OutResults);
diff --git a/Engine/Plugins/MovieScene/SequenceValidator/Source/SequenceValidator/Private/Validation/SequenceValidator.cpp b/Engine/Plugins/MovieScene/SequenceValidator/Source/SequenceValidator/Private/Validation/SequenceValidator.cpp
--- a/Engine/Plugins/MovieScene/SequenceValidator/Source/SequenceValidator/Private/Validation/SequenceValidator.cpp
+++ b/Engine/Plugins/MovieScene/SequenceValidator/Source/SequenceValidator/Private/Validation/SequenceValidator.cpp
@@ -308,6 +308,11 @@ void FSequenceValidator::Queue(TArrayView<UMovieSceneSequence*> InSequences)
 	}
 }
 
+void FSequenceValidator::Delete(UMovieSceneSequence* InSequence)
+{
+	ValidationQueue.Remove(InSequence);
+}
+
 void FSequenceValidator::ClearQueue()
 {
 	ValidationQueue.Reset();
diff --git a/Engine/Plugins/MovieScene/SequenceValidator/Source/SequenceValidator/Private/Widgets/SSequenceValidator.cpp b/Engine/Plugins/MovieScene/SequenceValidator/Source/SequenceValidator/Private/Widgets/SSequenceValidator.cpp
--- a/Engine/Plugins/MovieScene/SequenceValidator/Source/SequenceValidator/Private/Widgets/SSequenceValidator.cpp
+++ b/Engine/Plugins/MovieScene/SequenceValidator/Source/SequenceValidator/Private/Widgets/SSequenceValidator.cpp
@@ -5,10 +5,10 @@
 #include "Framework/Application/SlateApplication.h"
 #include "Framework/Docking/TabManager.h"
 #include "Framework/Views/ITypedTableView.h"
-#include "IContentBrowserSingleton.h"
 #include "MovieSceneSequence.h"
 #include "SequenceValidatorCommands.h"
 #include "SequenceValidatorStyle.h"
+#include "SPrimaryButton.h"
 #include "ToolMenus.h"
 #include "Validation/SequenceValidator.h"
 #include "Widgets/Docking/SDockTab.h"
@@ -28,7 +28,6 @@ namespace UE::Sequencer
 
 const FName SSequenceValidator::WindowName(TEXT("SequenceValidator"));
 const FName SSequenceValidator::MenubarName(TEXT("SequenceValidator.Menubar"));
-const FName SSequenceValidator::ToolbarName(TEXT("SequenceValidator.Toolbar"));
 
 void SSequenceValidator::RegisterTabSpawners()
 {
@@ -78,7 +77,6 @@ void SSequenceValidator::Construct(const FArguments& InArgs)
 
 	// Build all UI elements.
 	TSharedRef<SWidget> MenubarContents = ConstructMenubar();
-	TSharedRef<SWidget> ToolbarContents = ConstructToolbar(CommandList);
 
 	// Main layout.
 	ChildSlot
@@ -90,46 +88,56 @@ void SSequenceValidator::Construct(const FArguments& InArgs)
 			MenubarContents
 		]
 		+SVerticalBox::Slot()
-		.AutoHeight()
-		[
-			ToolbarContents
-		]
-		+SVerticalBox::Slot()
 		.Padding(2.0)
 		.FillHeight(1.f)
 		[
 			SNew(SSplitter)
 			.Orientation(Orient_Horizontal)
 			+SSplitter::Slot()
-			[
-				SAssignNew(QueueWidget, SSequenceValidatorQueue)
-				.Validator(Validator)
-			]
-			+SSplitter::Slot()
+			.Value(0.4)
 			[
 				SNew(SSplitter)
 				.Orientation(Orient_Vertical)
-				+SSplitter::Slot()
+				+ SSplitter::Slot()
 				[
-					SAssignNew(RulesWidget, SSequenceValidatorRules)
+					SAssignNew(QueueWidget, SSequenceValidatorQueue)
 					.Validator(Validator)
 				]
-				+SSplitter::Slot()
+				+ SSplitter::Slot()
 				[
-					SAssignNew(ResultsWidget, SSequenceValidatorResults)
+					SAssignNew(RulesWidget, SSequenceValidatorRules)
 					.Validator(Validator)
 				]
 			]
+			+SSplitter::Slot()
+			.Value(0.6)
+			[
+				SAssignNew(ResultsWidget, SSequenceValidatorResults)
+				.Validator(Validator)
+			]
 		]
 		+SVerticalBox::Slot()
 		.AutoHeight()
 		[
 			SNew(SHorizontalBox)
-			+SHorizontalBox::Slot()
+			+ SHorizontalBox::Slot()
+			.HAlign(HAlign_Left)
+			.Padding(4.f, 8.f)
 			[
 				SNew(STextBlock)
 				.Text(this, &SSequenceValidator::GetValidationStatusText)
 			]
+
+			+ SHorizontalBox::Slot()
+			.HAlign(HAlign_Right)
+			.Padding(4.f, 8.f)
+			[
+				SNew(SPrimaryButton)
+				.Text(LOCTEXT("StartValidation", "Start Validation"))
+				.ToolTipText(LOCTEXT("StartValidationTooltip", "Start validating the sequences in the queue"))
+				.OnClicked(FOnClicked::CreateLambda([this]() { StartValidation(); return FReply::Handled(); }))
+				.IsEnabled_Lambda([this]() -> bool { return this->CanStartValidation(); })
+			]
 		]
 	];
 }
@@ -149,91 +157,6 @@ TSharedRef<SWidget> SSequenceValidator::ConstructMenubar()
 	return ToolMenus->GenerateWidget(SSequenceValidator::MenubarName, MenubarContext);
 }
 
-TSharedRef<SWidget> SSequenceValidator::ConstructToolbar(TSharedRef<FUICommandList> InCommandList)
-{
-	UToolMenus* ToolMenus = UToolMenus::Get();
-	if (!ToolMenus->IsMenuRegistered(SSequenceValidator::ToolbarName))
-	{
-		FToolMenuOwnerScoped ToolMenuOwnerScope(this);
-
-		UToolMenu* Toolbar = ToolMenus->RegisterMenu(
-				SSequenceValidator::ToolbarName, NAME_None, EMultiBoxType::SlimHorizontalToolBar);
-
-		const FSequenceValidatorCommands& Commands = FSequenceValidatorCommands::Get();
-		TSharedRef<FSequenceValidatorStyle> ValidatorStyle = FSequenceValidatorStyle::Get();
-
-		FToolMenuSection& QueueSection = Toolbar->FindOrAddSection(TEXT("Queue"));
-		QueueSection.AddEntry(FToolMenuEntry::InitComboButton(
-			TEXT("Add"),
-			FUIAction(),
-			FNewToolMenuWidget::CreateStatic(&SSequenceValidator::GenerateAddSequenceMenu),
-			LOCTEXT("AddSequence_Label", "Add"),
-			LOCTEXT("AddSequence_ToolTip", "Add a sequence to the queue"),
-			FSlateIcon(ValidatorStyle->GetStyleSetName(), "SequenceValidator.QueueSequence")
-		));
-
-		FToolMenuSection& ValidationSection = Toolbar->FindOrAddSection(TEXT("Validation"));
-		ValidationSection.AddEntry(FToolMenuEntry::InitToolBarButton(Commands.StartValidation));
-	}
-
-	FToolMenuContext ToolbarContext(InCommandList, TSharedPtr<FExtender>());
-
-	USequenceValidatorMenuContext* Context = NewObject<USequenceValidatorMenuContext>();
-	Context->ValidatorWidget = SharedThis(this);
-	ToolbarContext.AddObject(Context);
-
-	return ToolMenus->GenerateWidget(SSequenceValidator::ToolbarName, ToolbarContext);
-}
-
-TSharedRef<SWidget> SSequenceValidator::GenerateAddSequenceMenu(const FToolMenuContext& Context)
-{
-	TSharedPtr<SSequenceValidator> This;
-	if (USequenceValidatorMenuContext* ValidatorContext = Context.FindContext<USequenceValidatorMenuContext>())
-	{
-		This = ValidatorContext->ValidatorWidget.Pin();
-	}
-
-	if (!This)
-	{
-		return SNullWidget::NullWidget;
-	}
-
-	FAssetPickerConfig PickerConfig;
-	PickerConfig.bAllowNullSelection = false;
-	PickerConfig.bFocusSearchBoxWhenOpened = true;
-	PickerConfig.Filter.bRecursiveClasses = true;
-	PickerConfig.Filter.ClassPaths.Add(UMovieSceneSequence::StaticClass()->GetClassPathName());
-	PickerConfig.SelectionMode = ESelectionMode::Single;
-	PickerConfig.InitialAssetViewType = EAssetViewType::List;
-	PickerConfig.OnAssetSelected = FOnAssetSelected::CreateSP(This.Get(), &SSequenceValidator::OnAddSequence);
-	PickerConfig.OnAssetsActivated = FOnAssetsActivated::CreateSP(This.Get(), &SSequenceValidator::OnAddSequences);
-
-	return IContentBrowserSingleton::Get().CreateAssetPicker(PickerConfig);
-}
-
-void SSequenceValidator::OnAddSequence(const FAssetData& SelectedAsset)
-{
-	if (UMovieSceneSequence* Sequence = Cast<UMovieSceneSequence>(SelectedAsset.GetAsset()))
-	{
-		Validator->Queue(Sequence);
-	}
-
-	QueueWidget->RequestListRefresh();
-}
-
-void SSequenceValidator::OnAddSequences(const TArray<FAssetData>& ActivatedAssets, EAssetTypeActivationMethod::Type ActivationMethod)
-{
-	for (const FAssetData& AssetData : ActivatedAssets)
-	{
-		if (UMovieSceneSequence* Sequence = Cast<UMovieSceneSequence>(AssetData.GetAsset()))
-		{
-			Validator->Queue(Sequence);
-		}
-	}
-
-	QueueWidget->RequestListRefresh();
-}
-
 void SSequenceValidator::Tick(const FGeometry& AllottedGeometry, const double InCurrentTime, const float InDeltaTime)
 {
 	SCompoundWidget::Tick(AllottedGeometry, InCurrentTime, InDeltaTime);
diff --git a/Engine/Plugins/MovieScene/SequenceValidator/Source/SequenceValidator/Private/Widgets/SSequenceValidator.h b/Engine/Plugins/MovieScene/SequenceValidator/Source/SequenceValidator/Private/Widgets/SSequenceValidator.h
--- a/Engine/Plugins/MovieScene/SequenceValidator/Source/SequenceValidator/Private/Widgets/SSequenceValidator.h
+++ b/Engine/Plugins/MovieScene/SequenceValidator/Source/SequenceValidator/Private/Widgets/SSequenceValidator.h
@@ -27,7 +27,6 @@ public:
 
 	static const FName WindowName;
 	static const FName MenubarName;
-	static const FName ToolbarName;
 
 	static void RegisterTabSpawners();
 	static TSharedRef<SDockTab> SpawnSequenceValidator(const FSpawnTabArgs& Args);
@@ -49,11 +48,6 @@ protected:
 private:
 
 	TSharedRef<SWidget> ConstructMenubar();
-	TSharedRef<SWidget> ConstructToolbar(TSharedRef<FUICommandList> InCommandList);
-
-	static TSharedRef<SWidget> GenerateAddSequenceMenu(const FToolMenuContext& Context);
-	void OnAddSequence(const FAssetData& SelectedAsset);
-	void OnAddSequences(const TArray<FAssetData>& ActivatedAssets, EAssetTypeActivationMethod::Type ActivationMethod);
 
 	void StartValidation();
 	bool CanStartValidation() const;
diff --git a/Engine/Plugins/MovieScene/SequenceValidator/Source/SequenceValidator/Private/Widgets/SSequenceValidatorQueue.cpp b/Engine/Plugins/MovieScene/SequenceValidator/Source/SequenceValidator/Private/Widgets/SSequenceValidatorQueue.cpp
--- a/Engine/Plugins/MovieScene/SequenceValidator/Source/SequenceValidator/Private/Widgets/SSequenceValidatorQueue.cpp
+++ b/Engine/Plugins/MovieScene/SequenceValidator/Source/SequenceValidator/Private/Widgets/SSequenceValidatorQueue.cpp
@@ -4,7 +4,10 @@
 
 #include "AssetThumbnail.h"
 #include "DragAndDrop/AssetDragDropOp.h"
+#include "IContentBrowserSingleton.h"
+#include "SequenceValidatorCommands.h"
 #include "SequenceValidatorStyle.h"
+#include "SPositiveActionButton.h"
 #include "ThumbnailRendering/ThumbnailManager.h"
 #include "Validation/SequenceValidator.h"
 #include "Widgets/SBoxPanel.h"
@@ -131,10 +134,44 @@ void SSequenceValidatorQueue::Construct(const FArguments& InArgs)
 		+SVerticalBox::Slot()
 		.AutoHeight()
 		[
-			SNew(STextBlock)
-			.Margin(4.f)
-			.Text(LOCTEXT("QueueTitle", "Queue"))
-			.TextStyle(ValidatorStyle, "SequenceValidator.PanelTitleText")
+			SNew(SBorder)
+			.Padding(0.f)
+			.BorderImage(FAppStyle::Get().GetBrush("ToolPanel.GroupBorder"))
+			[
+				SNew(SHorizontalBox)
+				+ SHorizontalBox::Slot()
+				.HAlign(HAlign_Left)
+				[
+					SNew(SHorizontalBox)
+					+ SHorizontalBox::Slot()
+					.AutoWidth()
+					.Padding(8.f, 4.f)
+					.VAlign(VAlign_Center)
+					[
+						SNew(SImage)
+						.Image(ValidatorStyle->GetBrush("SequenceValidator.QueueTitleIcon"))
+					]
+					+ SHorizontalBox::Slot()
+					.Padding(0.f, 4.f)
+					.FillWidth(1.f)
+					[
+						SNew(STextBlock)
+						.Margin(4.f)
+						.Text(LOCTEXT("QueueTitle", "Queue"))
+						.TextStyle(ValidatorStyle, "SequenceValidator.PanelTitleText")
+					]
+				]
+
+				+ SHorizontalBox::Slot()
+				.HAlign(HAlign_Right)
+				.Padding(8.f, 8.f)
+				[
+					SNew(SPositiveActionButton)
+					.Text(LOCTEXT("AddSequence_Label", "Add"))
+					.ToolTipText(LOCTEXT("AddSequence_ToolTip", "Add a sequence to the queue"))
+					.OnGetMenuContent(this, &SSequenceValidatorQueue::GenerateAddSequenceMenu)
+				]
+			]
 		]
 		+SVerticalBox::Slot()
 		.FillHeight(1.f)
@@ -142,16 +179,102 @@ void SSequenceValidatorQueue::Construct(const FArguments& InArgs)
 			SAssignNew(ListView, SListView<UMovieSceneSequence*>)
 			.ListItemsSource(&ItemSource)
 			.OnGenerateRow(this, &SSequenceValidatorQueue::OnListGenerateItemRow)
+			.OnContextMenuOpening(this, &SSequenceValidatorQueue::GetContextMenuContent)
 		]
 	];
 }
 
+
+TSharedRef<SWidget> SSequenceValidatorQueue::GenerateAddSequenceMenu()
+{
+	FAssetPickerConfig PickerConfig;
+	PickerConfig.bAllowNullSelection = false;
+	PickerConfig.bFocusSearchBoxWhenOpened = true;
+	PickerConfig.Filter.bRecursiveClasses = true;
+	PickerConfig.Filter.ClassPaths.Add(UMovieSceneSequence::StaticClass()->GetClassPathName());
+	PickerConfig.SelectionMode = ESelectionMode::Single;
+	PickerConfig.InitialAssetViewType = EAssetViewType::List;
+	PickerConfig.OnAssetSelected = FOnAssetSelected::CreateSP(this, &SSequenceValidatorQueue::OnAddSequence);
+	PickerConfig.OnAssetsActivated = FOnAssetsActivated::CreateSP(this, &SSequenceValidatorQueue::OnAddSequences);
+
+	return IContentBrowserSingleton::Get().CreateAssetPicker(PickerConfig);
+}
+
+void SSequenceValidatorQueue::OnAddSequence(const FAssetData& SelectedAsset)
+{
+	if (UMovieSceneSequence* Sequence = Cast<UMovieSceneSequence>(SelectedAsset.GetAsset()))
+	{
+		Validator->Queue(Sequence);
+	}
+
+	RequestListRefresh();
+}
+
+void SSequenceValidatorQueue::OnAddSequences(const TArray<FAssetData>& ActivatedAssets, EAssetTypeActivationMethod::Type ActivationMethod)
+{
+	for (const FAssetData& AssetData : ActivatedAssets)
+	{
+		if (UMovieSceneSequence* Sequence = Cast<UMovieSceneSequence>(AssetData.GetAsset()))
+		{
+			Validator->Queue(Sequence);
+		}
+	}
+
+	RequestListRefresh();
+}
+
+void SSequenceValidatorQueue::DeleteSequences()
+{
+	for (UMovieSceneSequence* Sequence : ListView->GetSelectedItems())
+	{
+		Validator->Delete(Sequence);
+	}
+	RequestListRefresh();
+}
+
+void SSequenceValidatorQueue::ClearQueue()
+{
+	Validator->ClearQueue();
+
+	RequestListRefresh();
+}
+
 TSharedRef<ITableRow> SSequenceValidatorQueue::OnListGenerateItemRow(UMovieSceneSequence* Item, const TSharedRef<STableViewBase>& OwnerTable)
 {
 	return SNew(SSequenceValidatorQueueEntry, OwnerTable)
 		.Item(Item);
 }
 
+TSharedPtr<SWidget> SSequenceValidatorQueue::GetContextMenuContent()
+{
+	FMenuBuilder MenuBuilder(true, nullptr);
+	MenuBuilder.BeginSection("Edit");
+	MenuBuilder.AddMenuEntry(
+		LOCTEXT("DeleteFromQueue", "Delete from Queue"),
+		LOCTEXT("DeleteFromQueue_Tooltip", "Delete the selected sequences from the queue"),
+		FSlateIcon(),
+		FUIAction(
+			FExecuteAction::CreateSP(this, &SSequenceValidatorQueue::DeleteSequences),
+			FCanExecuteAction::CreateLambda([this] { return ListView->GetNumItemsSelected() > 0; })
+		),
+		NAME_None,
+		EUserInterfaceActionType::Button
+	);
+	MenuBuilder.AddMenuEntry(
+		LOCTEXT("ClearQueue", "Clear Queue"),
+		LOCTEXT("ClearQueue_Tooltip", "Clear all the sequences from the queue"),
+		FSlateIcon(),
+		FUIAction(
+			FExecuteAction::CreateSP(this, &SSequenceValidatorQueue::ClearQueue)
+		),
+		NAME_None,
+		EUserInterfaceActionType::Button
+	);
+	MenuBuilder.EndSection();
+
+	return MenuBuilder.MakeWidget();
+}
+
 void SSequenceValidatorQueue::Tick(const FGeometry& AllottedGeometry, const double InCurrentTime, const float InDeltaTime)
 {
 	SCompoundWidget::Tick(AllottedGeometry, InCurrentTime, InDeltaTime);
diff --git a/Engine/Plugins/MovieScene/SequenceValidator/Source/SequenceValidator/Private/Widgets/SSequenceValidatorQueue.h b/Engine/Plugins/MovieScene/SequenceValidator/Source/SequenceValidator/Private/Widgets/SSequenceValidatorQueue.h
--- a/Engine/Plugins/MovieScene/SequenceValidator/Source/SequenceValidator/Private/Widgets/SSequenceValidatorQueue.h
+++ b/Engine/Plugins/MovieScene/SequenceValidator/Source/SequenceValidator/Private/Widgets/SSequenceValidatorQueue.h
@@ -2,6 +2,7 @@
 
 #pragma once
 
+#include "IAssetTypeActions.h"
 #include "MovieSceneSequence.h"
 #include "Widgets/SCompoundWidget.h"
 #include "Widgets/Views/SListView.h"
@@ -33,7 +34,14 @@ protected:
 
 private:
 
+	TSharedRef<SWidget> GenerateAddSequenceMenu();
+	void OnAddSequence(const FAssetData& SelectedAsset);
+	void OnAddSequences(const TArray<FAssetData>& ActivatedAssets, EAssetTypeActivationMethod::Type ActivationMethod);
+	void DeleteSequences();
+	void ClearQueue();
+
 	TSharedRef<ITableRow> OnListGenerateItemRow(UMovieSceneSequence* Item, const TSharedRef<STableViewBase>& OwnerTable);
+	TSharedPtr<SWidget> GetContextMenuContent();
 
 	void UpdateItemSource();
 
diff --git a/Engine/Plugins/MovieScene/SequenceValidator/Source/SequenceValidator/Private/Widgets/SSequenceValidatorResults.cpp b/Engine/Plugins/MovieScene/SequenceValidator/Source/SequenceValidator/Private/Widgets/SSequenceValidatorResults.cpp
--- a/Engine/Plugins/MovieScene/SequenceValidator/Source/SequenceValidator/Private/Widgets/SSequenceValidatorResults.cpp
+++ b/Engine/Plugins/MovieScene/SequenceValidator/Source/SequenceValidator/Private/Widgets/SSequenceValidatorResults.cpp
@@ -35,6 +35,7 @@ public:
 private:
 
 	const FSlateBrush* GetIcon() const;
+	FSlateColor GetColorAndOpacity() const;
 	FText GetMessage() const;
 
 private:
@@ -46,26 +47,83 @@ void SSequenceValidatorResultEntry::Construct(const FArguments& InArgs)
 {
 	WeakItem = InArgs._Item;
 
+	TSharedRef<FSequenceValidatorStyle> ValidatorStyle = FSequenceValidatorStyle::Get();
+
+	TSharedPtr<SHorizontalBox> ResultBox;
+
+	TSharedPtr<FSequenceValidationResult> Item = WeakItem.Pin();
+
+	if (Item->GetChildren().Num() > 0)
+	{
+		if (!ResultBox)
+		{
+			ResultBox = SNew(SHorizontalBox);
+		}
+
+		// Add all the rules once at the parent level
+		TSet<FString> RuleInfoNames;
+		for (TSharedPtr<FSequenceValidationResult> Child : Item->GetChildren())
+		{
+			const FSequenceValidationRuleInfo& RuleInfo = Child->GetRuleInfo();
+			if (RuleInfo.RuleName.IsEmpty() || RuleInfoNames.Contains(RuleInfo.RuleName.ToString()))
+			{
+				continue;
+			}
+			RuleInfoNames.Add(RuleInfo.RuleName.ToString());
+
+			ResultBox->AddSlot()
+			.Padding(FMargin(4, 0.f, 0.f, 0.f))
+			.AutoWidth()
+			[
+				SNew(SImage)
+				.Image(ValidatorStyle->GetBrush("ValidationRule.Rule"))
+				.ColorAndOpacity_Lambda([RuleInfo]{ return RuleInfo.RuleColor; })
+			];
+		}
+
+		if (RuleInfoNames.Num() == 0)
+		{
+			ResultBox->AddSlot()
+			.Padding(FMargin(4, 0.f, 0.f, 0.f))
+			.AutoWidth()
+			[
+				SNew(SImage)
+				.Image(ValidatorStyle->GetBrush("ValidationResult.Pass"))
+			];
+		}
+	}
+
 	ChildSlot
-	.Padding(0, 3, 5, 0)
+	.Padding(0, 4, 5, 4)
 	.VAlign(VAlign_Center)
 	[
 		SNew(SHorizontalBox)
 		+SHorizontalBox::Slot()
 		.AutoWidth()
 		.VAlign(VAlign_Center)
+		.Padding(2, 2)
 		[
 			SNew(SImage)
 			.Image(this, &SSequenceValidatorResultEntry::GetIcon)
-			.ColorAndOpacity(FSlateColor::UseForeground())
+			.ColorAndOpacity(this, &SSequenceValidatorResultEntry::GetColorAndOpacity)
 		]
 		+SHorizontalBox::Slot()
 		.AutoWidth()
 		.VAlign(VAlign_Center)
-		.Padding(2, 0)
+		.Padding(2, 2)
 		[
 			SNew(STextBlock)
 			.Text(this, &SSequenceValidatorResultEntry::GetMessage)
+			.AutoWrapText(true)
+		]
+		+ SHorizontalBox::Slot()
+		.HAlign(HAlign_Right)
+		.VAlign(VAlign_Center)
+		.Padding(2, 2)
+		[
+			ResultBox.IsValid()
+				? ResultBox.ToSharedRef()
+				: SNullWidget::NullWidget
 		]
 	];
 }
@@ -95,6 +153,16 @@ const FSlateBrush* SSequenceValidatorResultEntry::GetIcon() const
 	return nullptr;
 }
 
+FSlateColor SSequenceValidatorResultEntry::GetColorAndOpacity() const
+{
+	if (TSharedPtr<FSequenceValidationResult> Item = WeakItem.Pin())
+	{
+		return Item->GetRuleInfo().RuleColor;
+	}
+
+	return FSlateColor::UseForeground();
+}
+
 FText SSequenceValidatorResultEntry::GetMessage() const
 {
 	if (TSharedPtr<FSequenceValidationResult> Item = WeakItem.Pin())
@@ -126,10 +194,29 @@ void SSequenceValidatorResults::Construct(const FArguments& InArgs)
 		+SVerticalBox::Slot()
 		.AutoHeight()
 		[
-			SNew(STextBlock)
-			.Margin(4.f)
-			.Text(LOCTEXT("ResultsTitle", "Results"))
-			.TextStyle(ValidatorStyle, "SequenceValidator.PanelTitleText")
+			SNew(SBorder)
+			.Padding(0.f)
+			.BorderImage(FAppStyle::Get().GetBrush("ToolPanel.GroupBorder"))
+			[
+				SNew(SHorizontalBox)
+				+ SHorizontalBox::Slot()
+				.AutoWidth()
+				.Padding(8.f, 4.f)
+				.VAlign(VAlign_Center)
+				[
+					SNew(SImage)
+					.Image(ValidatorStyle->GetBrush("SequenceValidator.ResultsTitleIcon"))
+				]
+				+ SHorizontalBox::Slot()
+				.Padding(0.f, 4.f)
+				.FillWidth(1.f)
+				[
+					SNew(STextBlock)
+					.Margin(4.f)
+					.Text(LOCTEXT("ResultsTitle", "Results"))
+					.TextStyle(ValidatorStyle, "SequenceValidator.PanelTitleText")
+				]
+			]
 		]
 		+SVerticalBox::Slot()
 		.FillHeight(1.f)
@@ -263,7 +350,22 @@ void SSequenceValidatorResults::Tick(const FGeometry& AllottedGeometry, const do
 
 		for (FTreeItemPtr Item : ItemSource)
 		{
-			TreeView->SetItemExpansion(Item, true);
+			// Collapse any items that fully passed validation
+			bool bPassed = true;
+			if (Item->GetChildren().Num() > 0)
+			{
+				for (TSharedPtr<FSequenceValidationResult> Child : Item->GetChildren())
+				{
+					const FSequenceValidationRuleInfo& RuleInfo = Child->GetRuleInfo();
+					if (!RuleInfo.RuleName.IsEmpty())
+					{
+						bPassed = false;
+						break;
+					}
+				}
+			}
+
+			TreeView->SetItemExpansion(Item, !bPassed);
 		}
 	}
 }
diff --git a/Engine/Plugins/MovieScene/SequenceValidator/Source/SequenceValidator/Private/Widgets/SSequenceValidatorRules.cpp b/Engine/Plugins/MovieScene/SequenceValidator/Source/SequenceValidator/Private/Widgets/SSequenceValidatorRules.cpp
--- a/Engine/Plugins/MovieScene/SequenceValidator/Source/SequenceValidator/Private/Widgets/SSequenceValidatorRules.cpp
+++ b/Engine/Plugins/MovieScene/SequenceValidator/Source/SequenceValidator/Private/Widgets/SSequenceValidatorRules.cpp
@@ -27,6 +27,7 @@ private:
 
 	FText GetRuleName() const;
 	FText GetRuleDescription() const;
+	FSlateColor GetColorAndOpacity() const;
 	ECheckBoxState IsRuleEnabled() const;
 	void OnRuleEnabledChanged(ECheckBoxState CheckState) const;
 
@@ -53,7 +54,8 @@ void SSequenceValidatorRuleEntry::Construct(const FArguments& InArgs, const TSha
 		.VAlign(VAlign_Center)
 		[
 			SNew(SImage)
-			.Image(ValidatorStyle->GetBrush("ValidationRule"))
+			.Image(ValidatorStyle->GetBrush("ValidationRule.Rule"))
+			.ColorAndOpacity(this, &SSequenceValidatorRuleEntry::GetColorAndOpacity)
 		]
 		+SHorizontalBox::Slot()
 		.Padding(8.f, 4.f)
@@ -61,6 +63,7 @@ void SSequenceValidatorRuleEntry::Construct(const FArguments& InArgs, const TSha
 		[
 			SNew(SVerticalBox)
 			+SVerticalBox::Slot()
+			.Padding(2.f, 4.f)
 			[
 				SNew(STextBlock)
 				.Text(this, &SSequenceValidatorRuleEntry::GetRuleName)
@@ -68,10 +71,13 @@ void SSequenceValidatorRuleEntry::Construct(const FArguments& InArgs, const TSha
 				.ToolTipText(this, &SSequenceValidatorRuleEntry::GetRuleDescription)
 			]
 			+SVerticalBox::Slot()
+			.Padding(2.f, 4.f)
+			.AutoHeight()
 			[
 				SNew(STextBlock)
 				.Text(this, &SSequenceValidatorRuleEntry::GetRuleDescription)
 				.TextStyle(ValidatorStyle, "ValidatorRule.RuleDescriptionText")
+				.AutoWrapText(true)
 			]
 		]
 		+SHorizontalBox::Slot()
@@ -108,6 +114,15 @@ FText SSequenceValidatorRuleEntry::GetRuleDescription() const
 	return FText::GetEmpty();
 }
 
+FSlateColor SSequenceValidatorRuleEntry::GetColorAndOpacity() const
+{
+	if (TSharedPtr<FSequenceValidationRuleInfo> Item = WeakItem.Pin())
+	{
+		return Item->RuleColor;
+	}
+	return FSlateColor::UseForeground();
+}
+
 ECheckBoxState SSequenceValidatorRuleEntry::IsRuleEnabled() const
 {
 	if (TSharedPtr<FSequenceValidationRuleInfo> Item = WeakItem.Pin())
@@ -137,10 +152,29 @@ void SSequenceValidatorRules::Construct(const FArguments& InArgs)
 		+SVerticalBox::Slot()
 		.AutoHeight()
 		[
-			SNew(STextBlock)
-			.Margin(4.f)
-			.Text(LOCTEXT("RulesTitle", "Rules"))
-			.TextStyle(ValidatorStyle, "SequenceValidator.PanelTitleText")
+			SNew(SBorder)
+			.Padding(0.f)
+			.BorderImage(FAppStyle::Get().GetBrush("ToolPanel.GroupBorder"))
+			[
+				SNew(SHorizontalBox)
+				+ SHorizontalBox::Slot()
+				.AutoWidth()
+				.Padding(8.f, 4.f)
+				.VAlign(VAlign_Center)
+				[
+					SNew(SImage)
+					.Image(ValidatorStyle->GetBrush("SequenceValidator.RulesTitleIcon"))
+				]
+				+ SHorizontalBox::Slot()
+				.Padding(0.f, 4.f)
+				.FillWidth(1.f)
+				[
+					SNew(STextBlock)
+					.Margin(4.f)
+					.Text(LOCTEXT("RulesTitle", "Rules"))
+					.TextStyle(ValidatorStyle, "SequenceValidator.PanelTitleText")
+				]
+			]
 		]
 		+SVerticalBox::Slot()
 		.FillHeight(1.f)
diff --git a/Engine/Plugins/MovieScene/SequenceValidator/Source/SequenceValidator/Public/Validation/SequenceValidationResult.h b/Engine/Plugins/MovieScene/SequenceValidator/Source/SequenceValidator/Public/Validation/SequenceValidationResult.h
--- a/Engine/Plugins/MovieScene/SequenceValidator/Source/SequenceValidator/Public/Validation/SequenceValidationResult.h
+++ b/Engine/Plugins/MovieScene/SequenceValidator/Source/SequenceValidator/Public/Validation/SequenceValidationResult.h
@@ -7,6 +7,7 @@
 #include "Logging/TokenizedMessage.h"
 #include "Misc/FrameTime.h"
 #include "Misc/Optional.h"
+#include "SequenceValidationRule.h"
 #include "Templates/SharedPointer.h"
 #include "UObject/WeakObjectPtr.h"
 
@@ -31,9 +32,10 @@ public:
 	{}
 
 	/** Creates a new validation result for the given object, with the given severity. */
-	FSequenceValidationResult(EMessageSeverity::Type InSeverity, UObject* InTarget)
+	FSequenceValidationResult(EMessageSeverity::Type InSeverity, UObject* InTarget, const FSequenceValidationRuleInfo& InRuleInfo)
 		: Severity(InSeverity)
 		, WeakTarget(InTarget)
+		, RuleInfo(InRuleInfo)
 	{}
 
 	/** Creates a new validation result with the given severity and message. */
@@ -43,15 +45,17 @@ public:
 	{}
 
 	/** Creates a new validation result with the given severity and message, for the given object. */
-	FSequenceValidationResult(EMessageSeverity::Type InSeverity, UObject* InTarget, const FText& InUserMessage)
+	FSequenceValidationResult(EMessageSeverity::Type InSeverity, UObject* InTarget, const FText& InUserMessage, const FSequenceValidationRuleInfo& InRuleInfo)
 		: Severity(InSeverity)
 		, WeakTarget(InTarget)
 		, UserMessage(InUserMessage)
+		, RuleInfo(InRuleInfo)
 	{}
 
 public:
 
 	EMessageSeverity::Type GetSeverity() const { return Severity; }
+	const FSequenceValidationRuleInfo& GetRuleInfo() const { return RuleInfo; }
 
 	SEQUENCEVALIDATOR_API bool GetSubSectionTrail(TArray<UMovieSceneSubSection*>& OutTrail) const;
 	SEQUENCEVALIDATOR_API void SetSubSectionTrail(TArrayView<UMovieSceneSubSection*> InTrail);
@@ -96,6 +100,7 @@ private:
 	TWeakObjectPtr<> WeakTarget;
 	FText UserMessage;
 	TOptional<FFrameTime> LocalTime;
+	FSequenceValidationRuleInfo RuleInfo;
 };
 
 /**
diff --git a/Engine/Plugins/MovieScene/SequenceValidator/Source/SequenceValidator/Public/Validation/SequenceValidationRule.h b/Engine/Plugins/MovieScene/SequenceValidator/Source/SequenceValidator/Public/Validation/SequenceValidationRule.h
--- a/Engine/Plugins/MovieScene/SequenceValidator/Source/SequenceValidator/Public/Validation/SequenceValidationRule.h
+++ b/Engine/Plugins/MovieScene/SequenceValidator/Source/SequenceValidator/Public/Validation/SequenceValidationRule.h
@@ -3,6 +3,7 @@
 #pragma once
 
 #include "Internationalization/Text.h"
+#include "Styling/SlateColor.h"
 #include "Templates/SharedPointer.h"
 
 class UMovieSceneSequence;
@@ -26,8 +27,11 @@ struct FSequenceValidationRuleInfo
 	FOnCreateSequenceValidationRule RuleFactory;
 	/** Whether this rule is enabled. */
 	bool bIsEnabled = true;
+	/** Color to display in the UI */
+	FSlateColor RuleColor = FSlateColor::UseForeground();
 };
 
+
 /**
  * Base class for a validation rule that can be run on a sequence.
  */
@@ -42,6 +46,9 @@ public:
 
 	void Run(const UMovieSceneSequence* InSequence, FSequenceValidationResults& OutResults) const;
 
+	// Get the color for the rule from an internal palette
+	static FSlateColor GetRuleColor(const FText& InRuleName);
+
 protected:
 
 	virtual void OnRun(const UMovieSceneSequence* InSequence, FSequenceValidationResults& OutResults) const {}
diff --git a/Engine/Plugins/MovieScene/SequenceValidator/Source/SequenceValidator/Public/Validation/SequenceValidator.h b/Engine/Plugins/MovieScene/SequenceValidator/Source/SequenceValidator/Public/Validation/SequenceValidator.h
--- a/Engine/Plugins/MovieScene/SequenceValidator/Source/SequenceValidator/Public/Validation/SequenceValidator.h
+++ b/Engine/Plugins/MovieScene/SequenceValidator/Source/SequenceValidator/Public/Validation/SequenceValidator.h
@@ -45,6 +45,8 @@ public:
 	void Queue(UMovieSceneSequence* InSequence);
 	/** Add the given root sequences to the queue. */
 	void Queue(TArrayView<UMovieSceneSequence*> InSequences);
+	/** Remove a sequence from the queue. */
+	void Delete(UMovieSceneSequence* InSequence);
 	/** Clears the current queue. */
 	void ClearQueue();
 
diff --git a/Engine/Plugins/MovieScene/SequenceValidator/Source/SequenceValidator/SequenceValidator.Build.cs b/Engine/Plugins/MovieScene/SequenceValidator/Source/SequenceValidator/SequenceValidator.Build.cs
--- a/Engine/Plugins/MovieScene/SequenceValidator/Source/SequenceValidator/SequenceValidator.Build.cs
+++ b/Engine/Plugins/MovieScene/SequenceValidator/Source/SequenceValidator/SequenceValidator.Build.cs
@@ -25,6 +25,7 @@ public class SequenceValidator : ModuleRules
 				"Slate",
 				"SlateCore",
 				"ToolMenus",
+				"ToolWidgets",
 				"UnrealEd",
 				"WorkspaceMenuStructure",
 			}
