Sequencer: Nullptr check
#rnx

diff --git a/Engine/Plugins/MovieScene/SequenceValidator/Source/SequenceValidator/Private/Widgets/SSequenceValidatorResults.cpp b/Engine/Plugins/MovieScene/SequenceValidator/Source/SequenceValidator/Private/Widgets/SSequenceValidatorResults.cpp
--- a/Engine/Plugins/MovieScene/SequenceValidator/Source/SequenceValidator/Private/Widgets/SSequenceValidatorResults.cpp
+++ b/Engine/Plugins/MovieScene/SequenceValidator/Source/SequenceValidator/Private/Widgets/SSequenceValidatorResults.cpp
@@ -53,7 +53,7 @@ void SSequenceValidatorResultEntry::Construct(const FArguments& InArgs)
 
 	TSharedPtr<FSequenceValidationResult> Item = WeakItem.Pin();
 
-	if (Item->GetChildren().Num() > 0)
+	if (Item && Item->GetChildren().Num() > 0)
 	{
 		if (!ResultBox)
 		{
