Sequencer: Add visit object bindings/tracks

#rnx

diff --git a/Engine/Plugins/MovieScene/SequenceValidator/Source/SequenceValidator/Private/Validation/Rules/SequenceValidationRule_MissingAssets.cpp b/Engine/Plugins/MovieScene/SequenceValidator/Source/SequenceValidator/Private/Validation/Rules/SequenceValidationRule_MissingAssets.cpp
--- a/Engine/Plugins/MovieScene/SequenceValidator/Source/SequenceValidator/Private/Validation/Rules/SequenceValidationRule_MissingAssets.cpp
+++ b/Engine/Plugins/MovieScene/SequenceValidator/Source/SequenceValidator/Private/Validation/Rules/SequenceValidationRule_MissingAssets.cpp
@@ -126,7 +126,9 @@ void FSequenceValidationRule_MissingAssets::OnRun(const UMovieSceneSequence* InS
 	// Visit the sequence hierarchy and report on things that are off.
 	UE::MovieScene::FSequenceVisitParams VisitParams;
 	VisitParams.bVisitRootTracks = true;
+	VisitParams.bVisitObjectBindings = true;
 	VisitParams.bVisitSubSequences = true;
+	VisitParams.bVisitTracks = true;
 	VisitParams.bVisitSections = true;
 
 	FMissingAssetsValidationVisitor Validator(MakeRuleInfo(), OutResults);
