Sequence: Fix FSequenceValidationRuleInfo ownership so that they're created by FSequenceValidator and shared to SSequenceValidatorRules which needs to modify them when toggling enabled/disabled

Otherwise SSequenceValidatorRules had its own copy of the RuleInfo which wasn't toggling the version that FSequenceValidator was operating on.

#rnx

diff --git a/Engine/Plugins/MovieScene/SequenceValidator/Source/SequenceValidator/Private/Validation/SequenceValidator.cpp b/Engine/Plugins/MovieScene/SequenceValidator/Source/SequenceValidator/Private/Validation/SequenceValidator.cpp
--- a/Engine/Plugins/MovieScene/SequenceValidator/Source/SequenceValidator/Private/Validation/SequenceValidator.cpp
+++ b/Engine/Plugins/MovieScene/SequenceValidator/Source/SequenceValidator/Private/Validation/SequenceValidator.cpp
@@ -77,7 +77,7 @@ class FSequenceValidatorTaskScheduler
 {
 public:
 
-	FSequenceValidatorTaskScheduler(TArrayView<FSequenceValidationRuleInfo> InRuleInfos)
+	FSequenceValidatorTaskScheduler(TArray<TSharedPtr<FSequenceValidationRuleInfo>> InRuleInfos)
 		: ValidationRuleInfos(InRuleInfos)
 	{
 	}
@@ -99,7 +99,7 @@ private:
 
 private:
 
-	TArray<FSequenceValidationRuleInfo> ValidationRuleInfos;
+	TArray<TSharedPtr<FSequenceValidationRuleInfo>> ValidationRuleInfos;
 	FSequenceValidationResults Results;
 
 	// Transient fields for running a queue.
@@ -161,11 +161,11 @@ void FSequenceValidatorTaskScheduler::RunQueue(TArrayView<UMovieSceneSequence*>
 	Tasks.Reset();
 	for (const UMovieSceneSequence* Sequence : InQueue)
 	{
-		for (const FSequenceValidationRuleInfo& RuleInfo : ValidationRuleInfos)
+		for (TSharedPtr<FSequenceValidationRuleInfo> RuleInfo : ValidationRuleInfos)
 		{
-			if (RuleInfo.bIsEnabled)
+			if (RuleInfo && RuleInfo->bIsEnabled)
 			{
-				TSharedRef<FSequenceValidationRule> Rule = RuleInfo.RuleFactory.Execute();
+				TSharedRef<FSequenceValidationRule> Rule = RuleInfo->RuleFactory.Execute();
 				Tasks.Emplace(Sequence, Rule);
 			}
 		}
@@ -285,9 +285,14 @@ TArrayView<UMovieSceneSequence* const> FSequenceValidator::GetQueue() const
 	return ValidationQueue;
 }
 
-TConstArrayView<FSequenceValidationRuleInfo> FSequenceValidator::GetRules() const
+TArray<TSharedPtr<FSequenceValidationRuleInfo>> FSequenceValidator::GetRules() const
 {
-	return ValidationRuleInfos;
+	TArray<TSharedPtr<FSequenceValidationRuleInfo>> Rules;
+	for (TSharedRef<FSequenceValidationRuleInfo> ValidationRuleInfo : ValidationRuleInfos)
+	{
+		Rules.Add(ValidationRuleInfo.ToSharedPtr());
+	}
+	return Rules;
 }
 
 const FSequenceValidationResults& FSequenceValidator::GetResults() const
@@ -350,7 +355,11 @@ FSimpleMulticastDelegate& FSequenceValidator::GetOnValidationFinished()
 void FSequenceValidator::CreateValidationRuleInfos()
 {
 	ISequenceValidatorModule& ValidationModule = FModuleManager::LoadModuleChecked<ISequenceValidatorModule>("SequenceValidator");
-	ValidationRuleInfos = ValidationModule.GetValidationRules();
+	
+	for (const FSequenceValidationRuleInfo& ValidationRuleInfo : ValidationModule.GetValidationRules())
+	{
+		ValidationRuleInfos.Add(MakeShared<FSequenceValidationRuleInfo>(ValidationRuleInfo));
+	}
 }
 
 void FSequenceValidator::StartValidation(bool bAsync)
@@ -369,7 +378,7 @@ void FSequenceValidator::StartValidation(bool bAsync)
 
 	using namespace Internal;
 
-	Scheduler = MakeUnique<FSequenceValidatorTaskScheduler>(ValidationRuleInfos);
+	Scheduler = MakeUnique<FSequenceValidatorTaskScheduler>(GetRules());
 	
 	if (bAsync)
 	{
diff --git a/Engine/Plugins/MovieScene/SequenceValidator/Source/SequenceValidator/Private/Widgets/SSequenceValidatorRules.cpp b/Engine/Plugins/MovieScene/SequenceValidator/Source/SequenceValidator/Private/Widgets/SSequenceValidatorRules.cpp
--- a/Engine/Plugins/MovieScene/SequenceValidator/Source/SequenceValidator/Private/Widgets/SSequenceValidatorRules.cpp
+++ b/Engine/Plugins/MovieScene/SequenceValidator/Source/SequenceValidator/Private/Widgets/SSequenceValidatorRules.cpp
@@ -197,10 +197,13 @@ TSharedRef<ITableRow> SSequenceValidatorRules::OnListGenerateItemRow(FListItemPt
 void SSequenceValidatorRules::UpdateItemSource()
 {
 	ItemSource.Reset();
-	for (const FSequenceValidationRuleInfo& RuleInfo : Validator->GetRules())
+
+	for (TSharedPtr<FSequenceValidationRuleInfo> RuleInfo : Validator->GetRules())
 	{
-		TSharedRef<FSequenceValidationRuleInfo> NewRuleInfo = MakeShared<FSequenceValidationRuleInfo>(RuleInfo);
-		ItemSource.Add(NewRuleInfo);
+		if (RuleInfo)
+		{
+			ItemSource.Add(RuleInfo);
+		}
 	}
 }
 
diff --git a/Engine/Plugins/MovieScene/SequenceValidator/Source/SequenceValidator/Public/Validation/SequenceValidator.h b/Engine/Plugins/MovieScene/SequenceValidator/Source/SequenceValidator/Public/Validation/SequenceValidator.h
--- a/Engine/Plugins/MovieScene/SequenceValidator/Source/SequenceValidator/Public/Validation/SequenceValidator.h
+++ b/Engine/Plugins/MovieScene/SequenceValidator/Source/SequenceValidator/Public/Validation/SequenceValidator.h
@@ -37,7 +37,7 @@ public:
 	/** The current queue of sequences to validate. */
 	TArrayView<UMovieSceneSequence* const> GetQueue() const;
 	/** The list of validation rules to run on the queue. */
-	TConstArrayView<FSequenceValidationRuleInfo> GetRules() const;
+	TArray<TSharedPtr<FSequenceValidationRuleInfo>> GetRules() const;
 	/** The current list of results generated by the validation rules. */
 	const FSequenceValidationResults& GetResults() const;
 
@@ -79,7 +79,7 @@ private:
 private:
 
 	TArray<UMovieSceneSequence*> ValidationQueue;
-	TArray<FSequenceValidationRuleInfo> ValidationRuleInfos;
+	TArray<TSharedRef<FSequenceValidationRuleInfo>> ValidationRuleInfos;
 	FSequenceValidationResults ValidationResults;
 
 	TUniquePtr<Internal::FSequenceValidatorTaskScheduler> Scheduler;
