Sequencer: exclude skeletal animation sections from the validation rule about whole section end frames.

diff --git a/Engine/Plugins/MovieScene/SequenceValidator/Source/SequenceValidator/Private/Validation/Rules/SequenceValidationRule_WholeSectionRanges.cpp b/Engine/Plugins/MovieScene/SequenceValidator/Source/SequenceValidator/Private/Validation/Rules/SequenceValidationRule_WholeSectionRanges.cpp
--- a/Engine/Plugins/MovieScene/SequenceValidator/Source/SequenceValidator/Private/Validation/Rules/SequenceValidationRule_WholeSectionRanges.cpp
+++ b/Engine/Plugins/MovieScene/SequenceValidator/Source/SequenceValidator/Private/Validation/Rules/SequenceValidationRule_WholeSectionRanges.cpp
@@ -7,6 +7,7 @@
 #include "MovieSceneSection.h"
 #include "MovieSceneSequence.h"
 #include "MovieSceneSequenceVisitor.h"
+#include "Sections/MovieSceneSkeletalAnimationSection.h"
 #include "Validation/SequenceValidationResult.h"
 
 #define LOCTEXT_NAMESPACE "SequenceValidationRule_WholeSectionRanges"
@@ -24,6 +25,7 @@ public:
 		: RuleInfo(InRuleInfo)
 		, Results(OutResults)
 	{
+		UpperBoundExcludedSectionClasses.Add(UMovieSceneSkeletalAnimationSection::StaticClass());
 	}
 
 	virtual void VisitStartSubSequence(UMovieSceneSubSection* SubSection, const FGuid& ObjectBindingID, const FSubSequenceSpace& LocalSpace) override
@@ -52,7 +54,7 @@ public:
 				AddWarning(Section, StartTime);
 			}
 		}
-		if (SectionRange.HasUpperBound())
+		if (SectionRange.HasUpperBound() && ShouldSectionHaveWholeFrameUpperBound(Section))
 		{
 			const FFrameTime EndTime = ConvertFrameTime(SectionRange.GetUpperBoundValue(), TickResolution, DisplayRate);
 			if (EndTime.GetSubFrame() != 0.f)
@@ -64,6 +66,19 @@ public:
 
 private:
 
+	bool ShouldSectionHaveWholeFrameUpperBound(UMovieSceneSection* Section)
+	{
+		UClass* SectionClass = Section->GetClass();
+		for (TSubclassOf<UMovieSceneSection> Class : UpperBoundExcludedSectionClasses)
+		{
+			if (SectionClass->IsChildOf(Class))
+			{
+				return false;
+			}
+		}
+		return true;
+	}
+
 	void AddWarning(UMovieSceneSection* Section, FFrameTime LocalTime)
 	{
 		TSharedRef<FSequenceValidationResult> NewResult = MakeShared<FSequenceValidationResult>(EMessageSeverity::Warning, Section, RuleInfo);
@@ -75,6 +90,8 @@ private:
 
 private:
 
+	TArray<TSubclassOf<UMovieSceneSection>> UpperBoundExcludedSectionClasses;
+	
 	FSequenceValidationRuleInfo RuleInfo;
 	TArray<UMovieSceneSubSection*> CurrentSubSectionTrail;
 	FSequenceValidationResults& Results;
@@ -84,7 +101,7 @@ FSequenceValidationRuleInfo FSequenceValidationRule_WholeSectionRanges::MakeRule
 {
 	FSequenceValidationRuleInfo RuleInfo;
 	RuleInfo.RuleName = LOCTEXT("RuleName", "Whole Section Ranges");
-	RuleInfo.RuleDescription = LOCTEXT("RuleDescription", "Check that sections start and end on whole frames.");
+	RuleInfo.RuleDescription = LOCTEXT("RuleDescription", "Check that most sections start and end on whole frames.");
 	RuleInfo.RuleFactory = FOnCreateSequenceValidationRule::CreateLambda([]() { return MakeShared<FSequenceValidationRule_WholeSectionRanges>(); });
 	RuleInfo.RuleColor = FSequenceValidationRule::GetRuleColor(RuleInfo.RuleName);
 	return RuleInfo;
