SequenceValidator: fix comment on rule

diff --git a/Engine/Plugins/MovieScene/SequenceValidator/Source/SequenceValidator/Private/Validation/Rules/SequenceValidationRule_WholeSectionRanges.h b/Engine/Plugins/MovieScene/SequenceValidator/Source/SequenceValidator/Private/Validation/Rules/SequenceValidationRule_WholeSectionRanges.h
--- a/Engine/Plugins/MovieScene/SequenceValidator/Source/SequenceValidator/Private/Validation/Rules/SequenceValidationRule_WholeSectionRanges.h
+++ b/Engine/Plugins/MovieScene/SequenceValidator/Source/SequenceValidator/Private/Validation/Rules/SequenceValidationRule_WholeSectionRanges.h
@@ -8,7 +8,7 @@ namespace UE::Sequencer
 {
 
 /**
- * A validation rule that checks that sections have start/end times that are on whole frames.
+ * A validation rule that checks that most sections have start/end times that are on whole frames.
  */
 class FSequenceValidationRule_WholeSectionRanges : public FSequenceValidationRule
 {
