Sequence Validator: Add support for selecting keys on validation.

Selection is now cleared when clicking on a result, unless shift is held down

#rnx

diff --git a/Engine/Plugins/MovieScene/SequenceValidator/Source/SequenceValidator/Private/Widgets/SSequenceValidatorResults.cpp b/Engine/Plugins/MovieScene/SequenceValidator/Source/SequenceValidator/Private/Widgets/SSequenceValidatorResults.cpp
--- a/Engine/Plugins/MovieScene/SequenceValidator/Source/SequenceValidator/Private/Widgets/SSequenceValidatorResults.cpp
+++ b/Engine/Plugins/MovieScene/SequenceValidator/Source/SequenceValidator/Private/Widgets/SSequenceValidatorResults.cpp
@@ -16,6 +16,11 @@
 #include "Validation/SequenceValidationResult.h"
 #include "Validation/SequenceValidator.h"
 
+#include "MVVM/Selection/Selection.h"
+#include "MVVM/ViewModels/ChannelModel.h"
+#include "MVVM/ViewModels/SectionModel.h"
+#include "MVVM/ViewModels/SequencerEditorViewModel.h"
+
 #define LOCTEXT_NAMESPACE "SSequenceValidatorResults"
 
 namespace UE::Sequencer
@@ -300,6 +305,14 @@ void SSequenceValidatorResults::OnTreeViewMouseButtonDoubleClick(FTreeItemPtr It
 		return;
 	}
 
+	TSharedPtr<FSequencerSelection> Selection = Sequencer->GetViewModel()->GetSelection();
+
+	const bool bAdd = FSlateApplication::Get().GetModifierKeys().IsShiftDown();
+	if (!bAdd)
+	{
+		Selection->Empty();
+	}
+
 	UObject* Target = Item->GetTarget();
 	if (Target)
 	{
@@ -323,6 +336,41 @@ void SSequenceValidatorResults::OnTreeViewMouseButtonDoubleClick(FTreeItemPtr It
 			Sequencer->SelectSection(TargetSection);
 		}
 
+		FTargetKeys TargetKeys = Item->GetTargetKeys();
+		if (TargetKeys.Channel != nullptr)
+		{
+			TSharedPtr<FChannelModel> ChannelModel;
+
+			for (const FViewModelPtr TrackAreaItem : Selection->TrackArea)
+			{
+				if (const TViewModelPtr<FSectionModel> SectionModel = TrackAreaItem.ImplicitCast())
+				{
+					TParentFirstChildIterator<FChannelGroupModel> KeyAreaNodes = SectionModel->GetParentTrackModel().AsModel()->GetDescendantsOfType<FChannelGroupModel>();
+					for (const TViewModelPtr<FChannelGroupModel>& KeyAreaNode : KeyAreaNodes)
+					{
+						TSharedPtr<FChannelModel> KeyAreaChannelModel = KeyAreaNode->GetChannel(SectionModel);
+						if (KeyAreaChannelModel && KeyAreaChannelModel->GetChannel() == TargetKeys.Channel)
+						{
+							ChannelModel = KeyAreaChannelModel;
+							break;
+						}
+					}
+				}
+				if (ChannelModel)
+				{
+					break;
+				}
+			}
+
+			if (ChannelModel)
+			{
+				for (FKeyHandle KeyHandle : TargetKeys.KeyHandles)
+				{
+					Selection->KeySelection.Select(ChannelModel, KeyHandle);
+				}
+			}
+		}
+
 		// Move the time head to a time that relates to the message.
 		if (Item->HasLocalTime())
 		{
diff --git a/Engine/Plugins/MovieScene/SequenceValidator/Source/SequenceValidator/Public/Validation/SequenceValidationResult.h b/Engine/Plugins/MovieScene/SequenceValidator/Source/SequenceValidator/Public/Validation/SequenceValidationResult.h
--- a/Engine/Plugins/MovieScene/SequenceValidator/Source/SequenceValidator/Public/Validation/SequenceValidationResult.h
+++ b/Engine/Plugins/MovieScene/SequenceValidator/Source/SequenceValidator/Public/Validation/SequenceValidationResult.h
@@ -10,12 +10,20 @@
 #include "SequenceValidationRule.h"
 #include "Templates/SharedPointer.h"
 #include "UObject/WeakObjectPtr.h"
+#include "Channels/MovieSceneChannel.h"
 
 class UMovieSceneSubSection;
+class UMovieSceneSection;
 
 namespace UE::Sequencer
 {
 
+struct FTargetKeys
+{
+	FMovieSceneChannel* Channel;
+	TArray<FKeyHandle> KeyHandles;
+};
+
 /**
  * A single validation result, indicating a note, warning, error, etc. inside the sequence.
  */
@@ -63,6 +71,9 @@ public:
 	UObject* GetTarget() const { return WeakTarget.Get(); }
 	void SetTarget(UObject* InTarget) { WeakTarget = InTarget; }
 
+	FTargetKeys GetTargetKeys() const { return TargetKeys; }
+	void SetTargetKeys(FTargetKeys InTargetKeys) { TargetKeys = InTargetKeys; }
+
 	const FText& GetUserMessage() const { return UserMessage; }
 	void SetUserMessage(const FText& InUserMessage) { UserMessage = InUserMessage; }
 
@@ -98,6 +109,7 @@ private:
 	EMessageSeverity::Type Severity = EMessageSeverity::Info;
 	TArray<TWeakObjectPtr<UMovieSceneSubSection>> WeakSubSectionTrail;
 	TWeakObjectPtr<> WeakTarget;
+	FTargetKeys TargetKeys;
 	FText UserMessage;
 	TOptional<FFrameTime> LocalTime;
 	FSequenceValidationRuleInfo RuleInfo;
diff --git a/Engine/Plugins/MovieScene/SequenceValidator/Source/SequenceValidator/SequenceValidator.Build.cs b/Engine/Plugins/MovieScene/SequenceValidator/Source/SequenceValidator/SequenceValidator.Build.cs
--- a/Engine/Plugins/MovieScene/SequenceValidator/Source/SequenceValidator/SequenceValidator.Build.cs
+++ b/Engine/Plugins/MovieScene/SequenceValidator/Source/SequenceValidator/SequenceValidator.Build.cs
@@ -22,6 +22,7 @@ public class SequenceValidator : ModuleRules
 				"MovieSceneTracks",
 				"Projects",
 				"Sequencer",
+				"SequencerCore",
 				"Slate",
 				"SlateCore",
 				"ToolMenus",
