Sequencer Validator: Added duplicate keys validator

Checks for duplicate keys on the same channel, outputs a warning and selects the other duplicates.

#rnx

diff --git a/Engine/Plugins/MovieScene/SequenceValidator/Source/SequenceValidator/Private/SequenceValidatorModule.cpp b/Engine/Plugins/MovieScene/SequenceValidator/Source/SequenceValidator/Private/SequenceValidatorModule.cpp
--- a/Engine/Plugins/MovieScene/SequenceValidator/Source/SequenceValidator/Private/SequenceValidatorModule.cpp
+++ b/Engine/Plugins/MovieScene/SequenceValidator/Source/SequenceValidator/Private/SequenceValidatorModule.cpp
@@ -4,6 +4,7 @@
 
 #include "Editor.h"
 #include "SequenceValidatorCommands.h"
+#include "Validation/Rules/SequenceValidationRule_DuplicateKeys.h"
 #include "Validation/Rules/SequenceValidationRule_SectionAlignments.h"
 #include "Validation/Rules/SequenceValidationRule_UnassignedBindingsAndAssets.h"
 #include "Validation/Rules/SequenceValidationRule_WholeSectionRanges.h"
@@ -82,6 +83,8 @@ void FSequenceValidatorModule::RegisterBuiltInRules()
 		RegisterValidationRule(FSequenceValidationRule_UnassignedBindingsAndAssets::MakeRuleInfo()));
 	BuiltInValidationRules.Add(
 		RegisterValidationRule(FSequenceValidationRule_WholeSectionRanges::MakeRuleInfo()));
+	BuiltInValidationRules.Add(
+		RegisterValidationRule(FSequenceValidationRule_DuplicateKeys::MakeRuleInfo()));
 }
 
 void FSequenceValidatorModule::UnregisterBuiltInRules()
diff --git a/Engine/Plugins/MovieScene/SequenceValidator/Source/SequenceValidator/Private/Validation/Rules/SequenceValidationRule_DuplicateKeys.cpp b/Engine/Plugins/MovieScene/SequenceValidator/Source/SequenceValidator/Private/Validation/Rules/SequenceValidationRule_DuplicateKeys.cpp
new file mode 100644
--- /dev/null
+++ b/Engine/Plugins/MovieScene/SequenceValidator/Source/SequenceValidator/Private/Validation/Rules/SequenceValidationRule_DuplicateKeys.cpp
@@ -0,0 +1,132 @@
+// Copyright Epic Games, Inc. All Rights Reserved.
+
+#include "Validation/Rules/SequenceValidationRule_DuplicateKeys.h"
+
+#include "Channels/MovieSceneChannel.h"
+#include "Channels/MovieSceneChannelProxy.h"
+#include "Internationalization/Internationalization.h"
+#include "MovieScene.h"
+#include "MovieSceneCommonHelpers.h"
+#include "MovieSceneSection.h"
+#include "MovieSceneSequence.h"
+#include "MovieSceneSequenceVisitor.h"
+#include "Validation/SequenceValidationResult.h"
+
+#define LOCTEXT_NAMESPACE "SequenceValidationRule_DuplicateKeys"
+
+namespace UE::Sequencer
+{
+
+class FDuplicateKeysVisitor : public UE::MovieScene::ISequenceVisitor
+{
+	using FSubSequenceSpace = UE::MovieScene::FSubSequenceSpace;
+
+public:
+
+	FDuplicateKeysVisitor(const FSequenceValidationRuleInfo& InRuleInfo, FSequenceValidationResults& OutResults)
+		: RuleInfo(InRuleInfo)
+		, Results(OutResults)
+	{
+	}
+
+	virtual void VisitStartSubSequence(UMovieSceneSubSection* SubSection, const FGuid& ObjectBindingID, const FSubSequenceSpace& LocalSpace) override
+	{
+		CurrentSubSectionTrail.Add(SubSection);
+	}
+
+	virtual void VisitEndSubSequence(UMovieSceneSubSection* SubSection, const FGuid& ObjectBindingID, const FSubSequenceSpace& LocalSpace) override
+	{
+		ensure(CurrentSubSectionTrail.Num() > 0 && CurrentSubSectionTrail.Last() == SubSection);
+		CurrentSubSectionTrail.Pop(EAllowShrinking::No);
+	}
+
+	virtual void VisitSection(UMovieSceneTrack* Track, UMovieSceneSection* Section, const FGuid& ObjectBindingID, const FSubSequenceSpace& LocalSpace) override
+	{
+		for (const FMovieSceneChannelEntry& ChannelEntry : Section->GetChannelProxy().GetAllEntries())
+		{
+			TConstArrayView<FMovieSceneChannel*> Channels = ChannelEntry.GetChannels();
+			TConstArrayView<FMovieSceneChannelMetaData> ChannelMetaData = ChannelEntry.GetMetaData();
+
+			for (int32 ChannelIndex = 0; ChannelIndex < Channels.Num(); ++ChannelIndex)
+			{
+				if (FMovieSceneChannel* Channel = Channels[ChannelIndex])
+				{
+					TArray<FFrameNumber> KeyTimes;
+					TArray<FKeyHandle> KeyHandles;
+
+					Channel->GetKeys(TRange<FFrameNumber>::All(), &KeyTimes, &KeyHandles);
+
+					TSet<FFrameNumber> UniqueKeyTimes;
+					TArray<FKeyHandle> DuplicateKeyHandles;
+					FFrameNumber LocalTime;
+
+					for (int32 KeyIndex = 0; KeyIndex < KeyTimes.Num(); ++KeyIndex)
+					{
+						if (!UniqueKeyTimes.Contains(KeyTimes[KeyIndex]))
+						{
+							UniqueKeyTimes.Add(KeyTimes[KeyIndex]);
+						}
+						else
+						{
+							// Found duplicate!
+							DuplicateKeyHandles.Add(KeyHandles[KeyIndex]);
+							LocalTime = KeyTimes[KeyIndex];
+						}
+					}
+
+					if (DuplicateKeyHandles.Num() > 0)
+					{
+						const FMovieSceneChannelMetaData& MetaData = ChannelMetaData[ChannelIndex];
+						AddWarning(Section, Channel, MetaData, DuplicateKeyHandles, LocalTime);
+					}
+				}
+			}
+		}
+	}
+
+private:
+
+	void AddWarning(UMovieSceneSection* Section, FMovieSceneChannel* Channel, const FMovieSceneChannelMetaData& ChannelMetaData, const TArray<FKeyHandle>& KeyHandles, FFrameTime LocalTime)
+	{
+		TSharedRef<FSequenceValidationResult> NewResult = MakeShared<FSequenceValidationResult>(EMessageSeverity::Warning, Section, RuleInfo);
+		NewResult->SetSubSectionTrail(CurrentSubSectionTrail);
+		NewResult->SetTargetKeys(FTargetKeys(Channel, KeyHandles));
+		NewResult->SetLocalTime(LocalTime);
+		NewResult->SetUserMessage(FText::Format(LOCTEXT("DuplicateKeyWarningFmt", " Channel {0} : contains {1} duplicate key(s)"), MovieSceneHelpers::GetDisplayPathName(Section, Channel, ChannelMetaData), KeyHandles.Num()));
+		Results.AddResult(NewResult);
+	}
+
+private:
+
+	FSequenceValidationRuleInfo RuleInfo;
+	TArray<UMovieSceneSubSection*> CurrentSubSectionTrail;
+	FSequenceValidationResults& Results;
+};
+
+FSequenceValidationRuleInfo FSequenceValidationRule_DuplicateKeys::MakeRuleInfo()
+{
+	FSequenceValidationRuleInfo RuleInfo;
+	RuleInfo.RuleName = LOCTEXT("RuleName", "Duplicate Keys");
+	RuleInfo.RuleDescription = LOCTEXT("RuleDescription", "Check for duplicate keys on the same channel.");
+	RuleInfo.RuleFactory = FOnCreateSequenceValidationRule::CreateLambda([]() { return MakeShared<FSequenceValidationRule_DuplicateKeys>(); });
+	RuleInfo.RuleColor = FSequenceValidationRule::GetRuleColor(RuleInfo.RuleName);
+	return RuleInfo;
+}
+
+void FSequenceValidationRule_DuplicateKeys::OnRun(const UMovieSceneSequence* InSequence, FSequenceValidationResults& OutResults) const
+{
+	UE::MovieScene::FSequenceVisitParams VisitParams;
+	VisitParams.bVisitRootTracks = true;
+	VisitParams.bVisitObjectBindings = true;
+	VisitParams.bVisitSubSequences = true;
+	VisitParams.bVisitTracks = true;
+	VisitParams.bVisitSections = true;
+
+	FDuplicateKeysVisitor Visitor(MakeRuleInfo(), OutResults);
+	UE::MovieScene::VisitSequence(const_cast<UMovieSceneSequence*>(InSequence), VisitParams, Visitor);
+}
+
+}  // namespace UE::Sequencer
+
+#undef LOCTEXT_NAMESPACE
+
diff --git a/Engine/Plugins/MovieScene/SequenceValidator/Source/SequenceValidator/Private/Validation/Rules/SequenceValidationRule_DuplicateKeys.h b/Engine/Plugins/MovieScene/SequenceValidator/Source/SequenceValidator/Private/Validation/Rules/SequenceValidationRule_DuplicateKeys.h
new file mode 100644
--- /dev/null
+++ b/Engine/Plugins/MovieScene/SequenceValidator/Source/SequenceValidator/Private/Validation/Rules/SequenceValidationRule_DuplicateKeys.h
@@ -0,0 +1,25 @@
+// Copyright Epic Games, Inc. All Rights Reserved.
+
+#pragma once
+
+#include "Validation/SequenceValidationRule.h"
+
+namespace UE::Sequencer
+{
+
+/**
+ * A validation rule that warns about duplicate keys on the same channel.
+ */
+class FSequenceValidationRule_DuplicateKeys : public FSequenceValidationRule
+{
+public:
+
+	static FSequenceValidationRuleInfo MakeRuleInfo();
+
+protected:
+
+	virtual void OnRun(const UMovieSceneSequence* InSequence, FSequenceValidationResults& OutResults) const override;
+};
+
+}  // namespace UE::Sequencer
+
