Sequence Validator: Fix unitialized

#rnx

diff --git a/Engine/Plugins/MovieScene/SequenceValidator/Source/SequenceValidator/Public/Validation/SequenceValidationResult.h b/Engine/Plugins/MovieScene/SequenceValidator/Source/SequenceValidator/Public/Validation/SequenceValidationResult.h
--- a/Engine/Plugins/MovieScene/SequenceValidator/Source/SequenceValidator/Public/Validation/SequenceValidationResult.h
+++ b/Engine/Plugins/MovieScene/SequenceValidator/Source/SequenceValidator/Public/Validation/SequenceValidationResult.h
@@ -20,7 +20,7 @@ namespace UE::Sequencer
 
 struct FTargetKeys
 {
-	FMovieSceneChannel* Channel;
+	FMovieSceneChannel* Channel = nullptr;
 	TArray<FKeyHandle> KeyHandles;
 };
 
