PR #14125: SequenceValidator: Fix WholeSectionRanges rule reporting incorrect results

#Synced-CL 48964464

diff --git a/Engine/Plugins/MovieScene/SequenceValidator/Source/SequenceValidator/Private/Validation/Rules/SequenceValidationRule_WholeSectionRanges.cpp b/Engine/Plugins/MovieScene/SequenceValidator/Source/SequenceValidator/Private/Validation/Rules/SequenceValidationRule_WholeSectionRanges.cpp
--- a/Engine/Plugins/MovieScene/SequenceValidator/Source/SequenceValidator/Private/Validation/Rules/SequenceValidationRule_WholeSectionRanges.cpp
+++ b/Engine/Plugins/MovieScene/SequenceValidator/Source/SequenceValidator/Private/Validation/Rules/SequenceValidationRule_WholeSectionRanges.cpp
@@ -49,20 +49,22 @@ public:
 		TRange<FFrameNumber> SectionRange = Section->GetRange();
 		if (SectionRange.HasLowerBound())
 		{
-			const FFrameTime StartTime = ConvertFrameTime(SectionRange.GetLowerBoundValue(), TickResolution, DisplayRate);
+			const FFrameNumber LowerBoundValue = SectionRange.GetLowerBoundValue();
+			const FFrameTime StartTime = ConvertFrameTime(LowerBoundValue, TickResolution, DisplayRate);
 			if (StartTime.GetSubFrame() != 0.f)
 			{
 				const bool bIsLowerBound = true;
-				AddWarning(Section, StartTime, bIsLowerBound);
+				AddWarning(Section, LowerBoundValue, bIsLowerBound);
 			}
 		}
 		if (SectionRange.HasUpperBound() && ShouldSectionHaveWholeFrameUpperBound(Section))
 		{
-			const FFrameTime EndTime = ConvertFrameTime(SectionRange.GetUpperBoundValue(), TickResolution, DisplayRate);
+			const FFrameNumber UpperBoundValue = SectionRange.GetUpperBoundValue();
+			const FFrameTime EndTime = ConvertFrameTime(UpperBoundValue, TickResolution, DisplayRate);
 			if (EndTime.GetSubFrame() != 0.f)
 			{
-				const bool bIsLowerBound = true;
-				AddWarning(Section, EndTime, bIsLowerBound);
+				const bool bIsLowerBound = false;
+				AddWarning(Section, UpperBoundValue, bIsLowerBound);
 			}
 		}
 	}
